No longer compile UOL editor module with strict numeric type-safety

diff --git a/Engine/Source/Editor/UniversalObjectLocatorEditor/UniversalObjectLocatorEditor.Build.cs b/Engine/Source/Editor/UniversalObjectLocatorEditor/UniversalObjectLocatorEditor.Build.cs
--- a/Engine/Source/Editor/UniversalObjectLocatorEditor/UniversalObjectLocatorEditor.Build.cs
+++ b/Engine/Source/Editor/UniversalObjectLocatorEditor/UniversalObjectLocatorEditor.Build.cs
@@ -6,7 +6,8 @@ public class UniversalObjectLocatorEditor : ModuleRules
 {
 	public UniversalObjectLocatorEditor(ReadOnlyTargetRules Target) : base(Target)
 	{
-		UnsafeTypeCastWarningLevel = WarningLevel.Error;
+		// @todo: this is temporarily disabled due to some MovieScene code not compiling
+		//UnsafeTypeCastWarningLevel = WarningLevel.Error;
 
 		PublicDependencyModuleNames.AddRange(
 			new string[] {
