UOL editor improvements

UOLs can now be edited on a per-fragment basis. UOL fragments appear as a chain of items similar to a breadcrumb trail, and each can be edited individually.

UOL fragment editing now has an expanded API, including:
- The ability to provide a 'root context' for evaluating relative root fragments against.
- The ability to control what fragments are available based on context
- A handle API for interacting with editor-side fragment representations
- Editor-side knowledge of likely fragment output class and whether a fragment is absolute/relative

Added IUniversalObjectLocatorEditorModule::FindLocatorEditor so external systems can better interact with editor-side code.
Multi-select property editing now works correctly.

diff --git a/Engine/Source/Editor/UniversalObjectLocatorEditor/Private/ActorLocatorEditor.cpp b/Engine/Source/Editor/UniversalObjectLocatorEditor/Private/ActorLocatorEditor.cpp
--- a/Engine/Source/Editor/UniversalObjectLocatorEditor/Private/ActorLocatorEditor.cpp
+++ b/Engine/Source/Editor/UniversalObjectLocatorEditor/Private/ActorLocatorEditor.cpp
@@ -4,7 +4,7 @@
 #include "Modules/ModuleManager.h"
 #include "UniversalObjectLocatorFragmentTypeHandle.h"
 #include "UniversalObjectLocator.h"
-#include "UniversalObjectLocatorEditor.h"
+#include "UniversalObjectLocatorFragmentEditor.h"
 #include "IUniversalObjectLocatorEditorModule.h"
 #include "IUniversalObjectLocatorCustomization.h"
 
@@ -14,6 +14,8 @@
 #include "PropertyCustomizationHelpers.h"
 
 #include "GameFramework/Actor.h"
+#include "String/Split.h"
+#include "UniversalObjectLocators/ActorLocatorFragment.h"
 
 
 #define LOCTEXT_NAMESPACE "ActorLocatorEditor"
@@ -21,6 +23,11 @@
 namespace UE::UniversalObjectLocator
 {
 
+ELocatorFragmentEditorType FActorLocatorEditor::GetLocatorFragmentEditorType() const
+{
+	return ELocatorFragmentEditorType::Absolute;
+}
+
 bool FActorLocatorEditor::IsDragSupported(TSharedPtr<FDragDropOperation> DragOperation, UObject* Context) const
 {
 	TSharedPtr<FActorDragDropOp> ActorDrag;
@@ -77,58 +84,116 @@ UObject* FActorLocatorEditor::ResolveDragOperation(TSharedPtr<FDragDropOperation
 	return nullptr;
 }
 
-TSharedPtr<SWidget> FActorLocatorEditor::MakeEditUI(TSharedPtr<IUniversalObjectLocatorCustomization> Customization)
+TSharedPtr<SWidget> FActorLocatorEditor::MakeEditUI(const FEditUIParameters& InParameters)
 {
-	TSharedRef<SObjectPropertyEntryBox> EditWidget = SNew(SObjectPropertyEntryBox)
-	.ObjectPath(Customization.ToSharedRef(), &IUniversalObjectLocatorCustomization::GetPathToObject)
-	// Disabling this for now since it causes issues
-	//.PropertyHandle(Customization->GetProperty())
-	.AllowedClass(AActor::StaticClass())
-	.OnObjectChanged(this, &FActorLocatorEditor::OnSetObject, TWeakPtr<IUniversalObjectLocatorCustomization>(Customization))
-	.AllowClear(true)
-	.DisplayUseSelected(false)
-	.DisplayBrowse(true)
-	.DisplayThumbnail(true);
-
-	float MinWidth = 100.f;
-	float MaxWidth = 500.f;
-	EditWidget->GetDesiredWidth(MinWidth, MaxWidth);
-
-	return SNew(SBox)
-	.MinDesiredWidth(MinWidth)
-	.MaxDesiredWidth(MaxWidth)
-	[
-		EditWidget
-	];
+	AActor* InitialActor = GetActor(InParameters.Handle);
+	const bool bAllowClear = true;
+	const bool bAllowPickingLevelInstanceContent = true;
+	const FOnActorSelected OnActorSelected = FOnActorSelected::CreateSP(this, &FActorLocatorEditor::OnSetActor, TWeakPtr<IFragmentEditorHandle>(InParameters.Handle));
+	const FOnShouldFilterActor OnShouldFilterActor = FOnShouldFilterActor::CreateLambda([](const AActor*){ return true; });
+
+	return
+		SNew(SBox)
+		.MinDesiredWidth(400.0f)
+		.MaxDesiredWidth(400.0f)
+		[
+			PropertyCustomizationHelpers::MakeActorPickerWithMenu(
+				InitialActor,
+				bAllowClear,
+				bAllowPickingLevelInstanceContent,
+				OnShouldFilterActor,
+				OnActorSelected,
+				FSimpleDelegate(),
+				FSimpleDelegate())
+		];
 }
 
-FText FActorLocatorEditor::GetDisplayText() const
+AActor* FActorLocatorEditor::GetActor(TWeakPtr<IFragmentEditorHandle> InWeakHandle) const
 {
+	if (TSharedPtr<IFragmentEditorHandle> Handle = InWeakHandle.Pin())
+	{
+		const FUniversalObjectLocatorFragment& Fragment = Handle->GetFragment();
+		ensure(Fragment.GetFragmentTypeHandle() == FActorLocatorFragment::FragmentType);
+		const FActorLocatorFragment* Payload = Fragment.GetPayloadAs(FActorLocatorFragment::FragmentType);
+		return (AActor*)Payload->Path.ResolveObject();
+	}
+
+	return nullptr;
+}
+
+FText FActorLocatorEditor::GetDisplayText(const FUniversalObjectLocatorFragment* InFragment) const
+{
+	if(InFragment != nullptr)
+	{
+		ensure(InFragment->GetFragmentTypeHandle() == FActorLocatorFragment::FragmentType);
+		const FActorLocatorFragment* Payload = InFragment->GetPayloadAs(FActorLocatorFragment::FragmentType);
+		if(Payload)
+		{
+			const FString& SubPathString = Payload->Path.GetSubPathString();
+			if(!SubPathString.IsEmpty())
+			{
+				FStringView LevelStringView, ActorStringView;
+				UE::String::SplitLast(SubPathString, TEXT("."), LevelStringView, ActorStringView);
+				if(!ActorStringView.IsEmpty())
+				{
+					return FText::FromStringView(ActorStringView);
+				}
+			}
+		}
+	}
+
 	return LOCTEXT("ExternalActorLocatorName", "Actor");
 }
 
-FText FActorLocatorEditor::GetDisplayTooltip() const
+FText FActorLocatorEditor::GetDisplayTooltip(const FUniversalObjectLocatorFragment* InFragment) const
 {
-	return LOCTEXT("ExternalActorLocatorTooltip", "Change this to an actor reference");
+	if(InFragment != nullptr)
+	{
+		ensure(InFragment->GetFragmentTypeHandle() == FActorLocatorFragment::FragmentType);
+		const FActorLocatorFragment* Payload = InFragment->GetPayloadAs(FActorLocatorFragment::FragmentType);
+		if(Payload && Payload->Path.IsValid())
+		{
+			static const FTextFormat TextFormat(LOCTEXT("ExternalActorLocatorTooltipFormat", "A reference to actor {0}"));
+			return FText::Format(TextFormat, FText::FromString(Payload->Path.ToString()));
+		}
+	}
+
+	return LOCTEXT("ExternalActorLocatorTooltip", "An actor reference");
 }
 
-FSlateIcon FActorLocatorEditor::GetDisplayIcon() const
+FSlateIcon FActorLocatorEditor::GetDisplayIcon(const FUniversalObjectLocatorFragment* InFragment) const
 {
-	return FSlateIcon();
+	return FSlateIcon(FAppStyle::GetAppStyleSetName(), "ClassIcon.Actor");
 }
 
-void FActorLocatorEditor::OnSetObject(const FAssetData& InNewObject, TWeakPtr<IUniversalObjectLocatorCustomization> WeakCustomization)
+UClass* FActorLocatorEditor::ResolveClass(const FUniversalObjectLocatorFragment& InFragment, UObject* InContext) const
 {
-	TSharedPtr<IUniversalObjectLocatorCustomization> Customization = WeakCustomization.Pin();
-	if (!Customization)
+	if(UClass* Class = ILocatorFragmentEditor::ResolveClass(InFragment, InContext))
+	{
+		return Class;
+	}
+
+	return AActor::StaticClass();
+}
+
+void FActorLocatorEditor::OnSetActor(AActor* InActor, TWeakPtr<IFragmentEditorHandle> InWeakHandle)
+{
+	TSharedPtr<IFragmentEditorHandle> Handle = InWeakHandle.Pin();
+	if (!Handle)
 	{
 		return;
 	}
 
-	UObject* Object = InNewObject.FastGetAsset(true);
+	FUniversalObjectLocatorFragment NewFragment(FActorLocatorFragment::FragmentType);
+	FActorLocatorFragment* Payload = NewFragment.GetPayloadAs(FActorLocatorFragment::FragmentType);
+	Payload->Path = InActor;
+	Handle->SetValue(NewFragment);
+}
 
-	FUniversalObjectLocator NewRef(Object);
-	Customization->SetValue(MoveTemp(NewRef));
+FUniversalObjectLocatorFragment FActorLocatorEditor::MakeDefaultLocatorFragment() const
+{
+	FUniversalObjectLocatorFragment NewFragment(FActorLocatorFragment::FragmentType);
+	return NewFragment;
 }
 
 } // namespace UE::UniversalObjectLocator
diff --git a/Engine/Source/Editor/UniversalObjectLocatorEditor/Private/ActorLocatorEditor.h b/Engine/Source/Editor/UniversalObjectLocatorEditor/Private/ActorLocatorEditor.h
--- a/Engine/Source/Editor/UniversalObjectLocatorEditor/Private/ActorLocatorEditor.h
+++ b/Engine/Source/Editor/UniversalObjectLocatorEditor/Private/ActorLocatorEditor.h
@@ -3,33 +3,42 @@
 #pragma once
 
 #include "CoreTypes.h"
-#include "UniversalObjectLocatorEditor.h"
+#include "UniversalObjectLocatorFragmentEditor.h"
 #include "Templates/SharedPointerFwd.h"
 
 struct FAssetData;
+class AActor;
 
 namespace UE::UniversalObjectLocator
 {
 
 class IUniversalObjectLocatorCustomization;
 
-class FActorLocatorEditor : public ILocatorEditor
+class FActorLocatorEditor : public ILocatorFragmentEditor
 {
+	ELocatorFragmentEditorType GetLocatorFragmentEditorType() const override;
+
 	bool IsDragSupported(TSharedPtr<FDragDropOperation> DragOperation, UObject* Context) const override;
 
 	UObject* ResolveDragOperation(TSharedPtr<FDragDropOperation> DragOperation, UObject* Context) const override;
 
-	TSharedPtr<SWidget> MakeEditUI(TSharedPtr<IUniversalObjectLocatorCustomization> Customization) override;
+	TSharedPtr<SWidget> MakeEditUI(const FEditUIParameters& InParameters) override;
+
+	FText GetDisplayText(const FUniversalObjectLocatorFragment* InFragment = nullptr) const override;
 
-	FText GetDisplayText() const override;
+	FText GetDisplayTooltip(const FUniversalObjectLocatorFragment* InFragment = nullptr) const override;
 
-	FText GetDisplayTooltip() const override;
+	FSlateIcon GetDisplayIcon(const FUniversalObjectLocatorFragment* InFragment = nullptr) const override;
 
-	FSlateIcon GetDisplayIcon() const override;
+	UClass* ResolveClass(const FUniversalObjectLocatorFragment& InFragment, UObject* InContext) const override; 
+
+	FUniversalObjectLocatorFragment MakeDefaultLocatorFragment() const override;
 
 private:
 
-	void OnSetObject(const FAssetData& InNewObject, TWeakPtr<IUniversalObjectLocatorCustomization> WeakCustomization);
+	AActor* GetActor(TWeakPtr<IFragmentEditorHandle> InWeakHandle) const;
+
+	void OnSetActor(AActor* InActor, TWeakPtr<IFragmentEditorHandle> InWeakHandle);
 };
 
 } // namespace UE::UniversalObjectLocator
diff --git a/Engine/Source/Editor/UniversalObjectLocatorEditor/Private/AnimInstanceLocatorEditor.cpp b/Engine/Source/Editor/UniversalObjectLocatorEditor/Private/AnimInstanceLocatorEditor.cpp
--- a/Engine/Source/Editor/UniversalObjectLocatorEditor/Private/AnimInstanceLocatorEditor.cpp
+++ b/Engine/Source/Editor/UniversalObjectLocatorEditor/Private/AnimInstanceLocatorEditor.cpp
@@ -4,7 +4,7 @@
 #include "UniversalObjectLocatorFragmentTypeHandle.h"
 #include "UniversalObjectLocators/AnimInstanceLocatorFragment.h"
 #include "UniversalObjectLocator.h"
-#include "UniversalObjectLocatorEditor.h"
+#include "UniversalObjectLocatorFragmentEditor.h"
 #include "IUniversalObjectLocatorEditorModule.h"
 #include "IUniversalObjectLocatorCustomization.h"
 #include "ISequencerModule.h"
@@ -14,6 +14,7 @@
 #include "Framework/MultiBox/MultiBoxBuilder.h"
 
 #include "PropertyCustomizationHelpers.h"
+#include "Animation/AnimInstance.h"
 
 #include "GameFramework/Actor.h"
 #include "Components/SkeletalMeshComponent.h"
@@ -29,72 +30,10 @@ class SAnimInstanceLocatorEditorUI : public SCompoundWidget
 	SLATE_BEGIN_ARGS(SAnimInstanceLocatorEditorUI){}
 	SLATE_END_ARGS()
 
-	void Construct(const FArguments& InArgs, TSharedPtr<IUniversalObjectLocatorCustomization> InCustomization)
+	void Construct(const FArguments& InArgs, TSharedPtr<IFragmentEditorHandle> InHandle)
 	{
-		WeakCustomization = InCustomization;
+		WeakHandle = InHandle;
 
-		TSharedRef<SObjectPropertyEntryBox> EditWidget = SNew(SObjectPropertyEntryBox)
-		.ObjectPath(InCustomization.ToSharedRef(), &IUniversalObjectLocatorCustomization::GetPathToObject)
-		.AllowedClass(USkeletalMeshComponent::StaticClass())
-		.OnObjectChanged(this, &SAnimInstanceLocatorEditorUI::OnSetObject)
-		.AllowClear(true)
-		.DisplayUseSelected(true)
-		.DisplayBrowse(true)
-		.DisplayThumbnail(true);
-
-		float MinWidth = 100.f;
-		float MaxWidth = 500.f;
-		EditWidget->GetDesiredWidth(MinWidth, MaxWidth);
-
-		ChildSlot
-		[
-			SNew(SHorizontalBox)
-
-			+SHorizontalBox::Slot()
-			[
-				SNew(SBox)
-				.MinDesiredWidth(MinWidth)
-				.MaxDesiredWidth(MaxWidth)
-				[
-					EditWidget
-				]
-			]
-
-			+ SHorizontalBox::Slot()
-			.AutoWidth()
-			[
-				SNew(SComboButton)
-				.OnGetMenuContent(this, &SAnimInstanceLocatorEditorUI::CreateAnimInstanceTypeMenuContent)
-				.ButtonContent()
-				[
-					SNew(STextBlock)
-					.Text(this, &SAnimInstanceLocatorEditorUI::GetCurrentAnimInstanceTypeText)
-				]
-			]
-		];
-	}
-
-private:
-
-	void OnSetObject(const FAssetData& InNewObject)
-	{
-		TSharedPtr<IUniversalObjectLocatorCustomization> Customization = WeakCustomization.Pin();
-		if (!Customization)
-		{
-			return;
-		}
-
-		UObject* Object  = InNewObject.FastGetAsset(true);
-
-		if (Object)
-		{
-			FUniversalObjectLocator NewRef(Object);
-			Customization->SetValue(MoveTemp(NewRef));
-		}
-	}
-
-	TSharedRef<SWidget> CreateAnimInstanceTypeMenuContent()
-	{
 		const bool bCloseAfterSelection = true;
 		FMenuBuilder MenuBuilder(bCloseAfterSelection, nullptr);
 
@@ -124,9 +63,14 @@ private:
 			NAME_None,
 			EUserInterfaceActionType::RadioButton
 		);
-		return MenuBuilder.MakeWidget();
+		
+		ChildSlot
+		[
+			MenuBuilder.MakeWidget()
+		];
 	}
 
+private:
 	FText GetCurrentAnimInstanceTypeText() const
 	{
 		TOptional<EAnimInstanceLocatorFragmentType> CommonType = GetCurrentType();
@@ -144,29 +88,13 @@ private:
 
 	void ChangeType(EAnimInstanceLocatorFragmentType InType)
 	{
-		TSharedPtr<IUniversalObjectLocatorCustomization> Customization = WeakCustomization.Pin();
-		if (Customization)
+		TSharedPtr<IFragmentEditorHandle> Handle = WeakHandle.Pin();
+		if (Handle)
 		{
-			TSharedPtr<IPropertyHandle>    Property          = Customization->GetProperty();
-			const FStructProperty*         RawProperty       = CastFieldChecked<const FStructProperty>(Property->GetProperty());
-
-			Property->NotifyPreChange();
-			Property->EnumerateRawData([this, InType](void* RawData, const int32, const int32){
-
-				FUniversalObjectLocatorFragment* Locator = static_cast<FUniversalObjectLocator*>(RawData)->GetLastFragment();
-
-				FAnimInstanceLocatorFragment* Payload = nullptr;
-				if (Locator && Locator->TryGetPayloadAs(FAnimInstanceLocatorFragment::FragmentType, Payload))
-				{
-					Payload->Type = InType;
-				}
-
-				// Continue to the next property value
-				return true;
-			});
-
-			Property->NotifyPostChange(EPropertyChangeType::ValueSet);
-			Property->NotifyFinishedChangingProperties();
+			FUniversalObjectLocatorFragment Fragment(FAnimInstanceLocatorFragment::FragmentType);
+			FAnimInstanceLocatorFragment* Payload = Fragment.GetPayloadAs(FAnimInstanceLocatorFragment::FragmentType);
+			Payload->Type = InType;
+			Handle->SetValue(Fragment);
 		}
 	}
 
@@ -178,46 +106,27 @@ private:
 
 	TOptional<EAnimInstanceLocatorFragmentType> GetCurrentType() const
 	{
-		TOptional<EAnimInstanceLocatorFragmentType> CommonType;
+		TOptional<EAnimInstanceLocatorFragmentType> Type;
 
-		TSharedPtr<IUniversalObjectLocatorCustomization> Customization = WeakCustomization.Pin();
-		if (Customization)
+		TSharedPtr<IFragmentEditorHandle> Handle = WeakHandle.Pin();
+		if (Handle)
 		{
-			TSharedPtr<IPropertyHandle> Property    = Customization->GetProperty();
-			const FStructProperty*      RawProperty = CastFieldChecked<const FStructProperty>(Property->GetProperty());
-
-			Property->EnumerateRawData([this, &CommonType](void* RawData, const int32, const int32){
-
-				FUniversalObjectLocatorFragment* LastFragment = static_cast<FUniversalObjectLocator*>(RawData)->GetLastFragment();
-
-				FAnimInstanceLocatorFragment* Payload = nullptr;
-				if (LastFragment && LastFragment->TryGetPayloadAs(FAnimInstanceLocatorFragment::FragmentType, Payload))
-				{
-					if (!CommonType.IsSet())
-					{
-						CommonType = Payload->Type;
-					}
-					else if (CommonType.GetValue() != Payload->Type)
-					{
-						// Different type - common type is undefined
-						CommonType.Reset();
-						// Stop iterating and return an unset optional
-						return false;
-					}
-				}
-
-				// Continue to the next property value
-				return true;
-			});
+			const FUniversalObjectLocatorFragment& Fragment = Handle->GetFragment();
+			const FAnimInstanceLocatorFragment* Payload = Fragment.GetPayloadAs(FAnimInstanceLocatorFragment::FragmentType);
+			Type = Payload->Type;
 		}
 
-		return CommonType;
+		return Type;
 	}
 
-	TWeakPtr<IUniversalObjectLocatorCustomization> WeakCustomization;
-	EAnimInstanceLocatorFragmentType CurrentType;
+	TWeakPtr<IFragmentEditorHandle> WeakHandle;
 };
 
+ELocatorFragmentEditorType FAnimInstanceLocatorEditor::GetLocatorFragmentEditorType() const
+{
+	return ELocatorFragmentEditorType::Relative;
+}
+
 bool FAnimInstanceLocatorEditor::IsDragSupported(TSharedPtr<FDragDropOperation> DragOperation, UObject* Context) const
 {
 	TSharedPtr<FActorDragDropOp> ActorDrag;
@@ -230,7 +139,7 @@ bool FAnimInstanceLocatorEditor::IsDragSupported(TSharedPtr<FDragDropOperation>
 	else if (DragOperation->IsOfType<FActorDragDropOp>())
 	{
 		ActorDrag = StaticCastSharedPtr<FActorDragDropOp>(DragOperation);
-	}		
+	}
 
 	if (ActorDrag)
 	{
@@ -274,25 +183,83 @@ UObject* FAnimInstanceLocatorEditor::ResolveDragOperation(TSharedPtr<FDragDropOp
 	return nullptr;
 }
 
-TSharedPtr<SWidget> FAnimInstanceLocatorEditor::MakeEditUI(TSharedPtr<IUniversalObjectLocatorCustomization> Customization)
+TSharedPtr<SWidget> FAnimInstanceLocatorEditor::MakeEditUI(const FEditUIParameters& InParameters)
 {
-	return SNew(SAnimInstanceLocatorEditorUI, Customization);
+	return SNew(SAnimInstanceLocatorEditorUI, InParameters.Handle);
 }
 
-FText FAnimInstanceLocatorEditor::GetDisplayText() const
+FText FAnimInstanceLocatorEditor::GetDisplayText(const FUniversalObjectLocatorFragment* InFragment) const
 {
-	return LOCTEXT("AnimInstanceLocatorName", "Anim Instance");
+	if(InFragment != nullptr)
+	{
+		ensure(InFragment->GetFragmentTypeHandle() == FAnimInstanceLocatorFragment::FragmentType);
+		const FAnimInstanceLocatorFragment* AnimInstanceLocatorFragment = InFragment->GetPayloadAs(FAnimInstanceLocatorFragment::FragmentType);
+		if(AnimInstanceLocatorFragment)
+		{
+			if(AnimInstanceLocatorFragment->Type == EAnimInstanceLocatorFragmentType::AnimInstance)
+			{
+				return LOCTEXT("AnimInstanceLocatorLabel", "Anim Instance");
+			}
+			else
+			{
+				return LOCTEXT("PostProcessInstanceLocatorLabel", "Post-process Anim Instance");
+			}
+		}
+	}
+	return LOCTEXT("AnimInstanceLocatorLabel", "Anim Instance");
+}
+
+FText FAnimInstanceLocatorEditor::GetDisplayTooltip(const FUniversalObjectLocatorFragment* InFragment) const
+{
+	if(InFragment != nullptr)
+	{
+		ensure(InFragment->GetFragmentTypeHandle() == FAnimInstanceLocatorFragment::FragmentType);
+		const FAnimInstanceLocatorFragment* AnimInstanceLocatorFragment = InFragment->GetPayloadAs(FAnimInstanceLocatorFragment::FragmentType);
+		if(AnimInstanceLocatorFragment)
+		{
+			if(AnimInstanceLocatorFragment->Type == EAnimInstanceLocatorFragmentType::AnimInstance)
+			{
+				return LOCTEXT("AnimInstanceLocatorTooltip", "A reference to an Anim Instance");
+			}
+			else
+			{
+				return LOCTEXT("PostProcessInstanceLocatorTooltip", "A reference to a Post-process Anim Instance");
+			}
+		}
+	}
+
+	return LOCTEXT("AnimInstanceLocatorTooltip", "A reference to an Anim Instance");
 }
-FText FAnimInstanceLocatorEditor::GetDisplayTooltip() const
+
+FSlateIcon FAnimInstanceLocatorEditor::GetDisplayIcon(const FUniversalObjectLocatorFragment* InFragment) const
 {
-	return LOCTEXT("AnimInstanceLocatorTooltip", "Change this to a reference to an Anim Instance");
+	return FSlateIcon(FAppStyle::GetAppStyleSetName(), "ClassIcon.AnimBlueprint");
 }
-FSlateIcon FAnimInstanceLocatorEditor::GetDisplayIcon() const
+
+UClass* FAnimInstanceLocatorEditor::ResolveClass(const FUniversalObjectLocatorFragment& InFragment, UObject* InContext) const
 {
-	return FSlateIcon();
+	if(UClass* Class = ILocatorFragmentEditor::ResolveClass(InFragment, InContext))
+	{
+		return Class;
+	}
+
+	if(USkeletalMeshComponent* SkeletalMeshComponent = Cast<USkeletalMeshComponent>(InContext))
+	{
+		if(SkeletalMeshComponent->AnimClass != nullptr)
+		{
+			return SkeletalMeshComponent->AnimClass;
+		}
+	}
+
+	return UAnimInstance::StaticClass();
 }
 
-} // namespace UE::UniversalObjectLocator
+FUniversalObjectLocatorFragment FAnimInstanceLocatorEditor::MakeDefaultLocatorFragment() const
+{
+	FUniversalObjectLocatorFragment NewFragment(FAnimInstanceLocatorFragment::FragmentType);
+	return NewFragment;
+}
 
+} // namespace UE::UniversalObjectLocator
 
 #undef LOCTEXT_NAMESPACE
diff --git a/Engine/Source/Editor/UniversalObjectLocatorEditor/Private/AnimInstanceLocatorEditor.h b/Engine/Source/Editor/UniversalObjectLocatorEditor/Private/AnimInstanceLocatorEditor.h
--- a/Engine/Source/Editor/UniversalObjectLocatorEditor/Private/AnimInstanceLocatorEditor.h
+++ b/Engine/Source/Editor/UniversalObjectLocatorEditor/Private/AnimInstanceLocatorEditor.h
@@ -3,25 +3,31 @@
 #pragma once
 
 #include "CoreTypes.h"
-#include "UniversalObjectLocatorEditor.h"
+#include "UniversalObjectLocatorFragmentEditor.h"
 
 namespace UE::UniversalObjectLocator
 {
 
-class FAnimInstanceLocatorEditor : public ILocatorEditor
+class FAnimInstanceLocatorEditor : public ILocatorFragmentEditor
 {
 public:
+	ELocatorFragmentEditorType GetLocatorFragmentEditorType() const override;
+	
 	bool IsDragSupported(TSharedPtr<FDragDropOperation> DragOperation, UObject* Context) const override;
 
 	UObject* ResolveDragOperation(TSharedPtr<FDragDropOperation> DragOperation, UObject* Context) const override;
 
-	TSharedPtr<SWidget> MakeEditUI(TSharedPtr<IUniversalObjectLocatorCustomization> Customization) override;
+	TSharedPtr<SWidget> MakeEditUI(const FEditUIParameters& InParameters) override;
 
-	FText GetDisplayText() const override;
+	FText GetDisplayText(const FUniversalObjectLocatorFragment* InFragment) const override;
 
-	FText GetDisplayTooltip() const override;
+	FText GetDisplayTooltip(const FUniversalObjectLocatorFragment* InFragment) const override;
 
-	FSlateIcon GetDisplayIcon() const override;
+	FSlateIcon GetDisplayIcon(const FUniversalObjectLocatorFragment* InFragment) const override;
+
+	UClass* ResolveClass(const FUniversalObjectLocatorFragment& InFragment, UObject* InContext) const override;
+
+	FUniversalObjectLocatorFragment MakeDefaultLocatorFragment() const override;
 };
 
 
diff --git a/Engine/Source/Editor/UniversalObjectLocatorEditor/Private/AssetLocatorEditor.cpp b/Engine/Source/Editor/UniversalObjectLocatorEditor/Private/AssetLocatorEditor.cpp
--- a/Engine/Source/Editor/UniversalObjectLocatorEditor/Private/AssetLocatorEditor.cpp
+++ b/Engine/Source/Editor/UniversalObjectLocatorEditor/Private/AssetLocatorEditor.cpp
@@ -4,13 +4,16 @@
 #include "Modules/ModuleManager.h"
 #include "UniversalObjectLocatorFragmentTypeHandle.h"
 #include "UniversalObjectLocator.h"
-#include "UniversalObjectLocatorEditor.h"
+#include "UniversalObjectLocatorFragmentEditor.h"
 #include "IUniversalObjectLocatorEditorModule.h"
 #include "IUniversalObjectLocatorCustomization.h"
 
 #include "DragAndDrop/AssetDragDropOp.h"
 
 #include "PropertyCustomizationHelpers.h"
+#include "AssetRegistry/AssetRegistryModule.h"
+#include "Styling/SlateIconFinder.h"
+#include "UniversalObjectLocators/AssetLocatorFragment.h"
 
 #include "Widgets/Layout/SBox.h"
 
@@ -20,6 +23,11 @@
 namespace UE::UniversalObjectLocator
 {
 
+ELocatorFragmentEditorType FAssetLocatorEditor::GetLocatorFragmentEditorType() const
+{
+	return ELocatorFragmentEditorType::Absolute;
+}
+
 bool FAssetLocatorEditor::IsDragSupported(TSharedPtr<FDragDropOperation> DragOperation, UObject* Context) const
 {
 	TSharedPtr<FAssetDragDropOp> ActorDrag;
@@ -52,58 +60,127 @@ UObject* FAssetLocatorEditor::ResolveDragOperation(TSharedPtr<FDragDropOperation
 	return nullptr;
 }
 
-TSharedPtr<SWidget> FAssetLocatorEditor::MakeEditUI(TSharedPtr<IUniversalObjectLocatorCustomization> Customization)
+TSharedPtr<SWidget> FAssetLocatorEditor::MakeEditUI(const FEditUIParameters& InParameters)
 {
-	TSharedRef<SObjectPropertyEntryBox> EditWidget = SNew(SObjectPropertyEntryBox)
-	.ObjectPath(Customization.ToSharedRef(), &IUniversalObjectLocatorCustomization::GetPathToObject)
-	.PropertyHandle(Customization->GetProperty())
-	.AllowedClass(UObject::StaticClass())
-	.OnObjectChanged(this, &FAssetLocatorEditor::OnSetObject, TWeakPtr<IUniversalObjectLocatorCustomization>(Customization))
-	.AllowClear(true)
-	.DisplayUseSelected(true)
-	.DisplayBrowse(true)
-	.DisplayThumbnail(true);
-
-	float MinWidth = 100.f;
-	float MaxWidth = 500.f;
-	EditWidget->GetDesiredWidth(MinWidth, MaxWidth);
-
-	return SNew(SBox)
-	.MinDesiredWidth(MinWidth)
-	.MaxDesiredWidth(MaxWidth)
-	[
-		EditWidget
-	];
+	FAssetData InitialAsset = GetAsset(InParameters.Handle);
+	const bool bAllowClear = true;
+	const bool bAllowCopyPaste = true;
+	const TArray<const UClass*> AllowedClasses = { UObject::StaticClass() };
+	const FOnShouldFilterAsset OnShouldFilterAsset = FOnShouldFilterAsset::CreateLambda([](const FAssetData&){ return false; });
+	const FOnAssetSelected OnAssetSelected = FOnAssetSelected::CreateSP(this, &FAssetLocatorEditor::OnSetAsset, TWeakPtr<IFragmentEditorHandle>(InParameters.Handle));
+
+	return
+		SNew(SBox)
+		.MinDesiredWidth(400.0f)
+		.MaxDesiredWidth(400.0f)
+		[
+			PropertyCustomizationHelpers::MakeAssetPickerWithMenu(
+				InitialAsset,
+				bAllowClear,
+				bAllowCopyPaste,
+				AllowedClasses,
+				TArray<const UClass*>(),
+				TArray<UFactory*>(),
+				OnShouldFilterAsset,
+				OnAssetSelected,
+				FSimpleDelegate(),
+				nullptr,
+				TArray<FAssetData>())
+		];
+}
+
+FAssetData FAssetLocatorEditor::GetAsset(TWeakPtr<IFragmentEditorHandle> InWeakHandle) const
+{
+	if (TSharedPtr<IFragmentEditorHandle> Handle = InWeakHandle.Pin())
+	{
+		const FUniversalObjectLocatorFragment& Fragment = Handle->GetFragment();
+		ensure(Fragment.GetFragmentTypeHandle() == FAssetLocatorFragment::FragmentType);
+		const FAssetLocatorFragment* Payload = Fragment.GetPayloadAs(FAssetLocatorFragment::FragmentType);
+		const IAssetRegistry& AssetRegistry = FModuleManager::Get().LoadModuleChecked<FAssetRegistryModule>("AssetRegistry").Get();
+		return AssetRegistry.GetAssetByObjectPath(FSoftObjectPath(Payload->Path));
+	}
+
+	return FAssetData();
 }
 
-FText FAssetLocatorEditor::GetDisplayText() const
+FText FAssetLocatorEditor::GetDisplayText(const FUniversalObjectLocatorFragment* InFragment) const
 {
+	if(InFragment != nullptr)
+	{
+		ensure(InFragment->GetFragmentTypeHandle() == FAssetLocatorFragment::FragmentType);
+		const FAssetLocatorFragment* Payload = InFragment->GetPayloadAs(FAssetLocatorFragment::FragmentType);
+		if(Payload)
+		{
+			return FText::FromName(Payload->Path.GetAssetName());
+		}
+	}
+
 	return LOCTEXT("AssetLocatorName", "Asset");
 }
 
-FText FAssetLocatorEditor::GetDisplayTooltip() const
+FText FAssetLocatorEditor::GetDisplayTooltip(const FUniversalObjectLocatorFragment* InFragment) const
 {
-	return LOCTEXT("AssetLocatorTooltip", "Change this to an asset reference");
+	if(InFragment != nullptr)
+	{
+		ensure(InFragment->GetFragmentTypeHandle() == FAssetLocatorFragment::FragmentType);
+		const FAssetLocatorFragment* Payload = InFragment->GetPayloadAs(FAssetLocatorFragment::FragmentType);
+		if(Payload)
+		{
+			static const FTextFormat TextFormat(LOCTEXT("AssetLocatorTooltipFormat", "A reference to asset {0}"));
+			return FText::Format(TextFormat, FText::FromString(Payload->Path.ToString()));
+		}
+	}
+
+	return LOCTEXT("AssetLocatorTooltip", "An asset reference");
+}
+
+FSlateIcon FAssetLocatorEditor::GetDisplayIcon(const FUniversalObjectLocatorFragment* InFragment) const
+{
+	if(InFragment != nullptr)
+	{
+		ensure(InFragment->GetFragmentTypeHandle() == FAssetLocatorFragment::FragmentType);
+		const FAssetLocatorFragment* Payload = InFragment->GetPayloadAs(FAssetLocatorFragment::FragmentType);
+		if(Payload)
+		{
+			IAssetRegistry& AssetRegistry = FModuleManager::LoadModuleChecked<FAssetRegistryModule>("AssetRegistry").Get();
+			FAssetData AssetData = AssetRegistry.GetAssetByObjectPath(FSoftObjectPath(Payload->Path));
+			return FSlateIconFinder::FindIconForClass(AssetData.GetClass());
+		}
+	}
+
+	return FSlateIcon(FAppStyle::GetAppStyleSetName(), "ClassIcon.Object");
 }
 
-FSlateIcon FAssetLocatorEditor::GetDisplayIcon() const
+UClass* FAssetLocatorEditor::ResolveClass(const FUniversalObjectLocatorFragment& InFragment, UObject* InContext) const
 {
-	return FSlateIcon();
+	if(UClass* Class = ILocatorFragmentEditor::ResolveClass(InFragment, InContext))
+	{
+		return Class;
+	}
+
+	return UObject::StaticClass();
 }
 
-void FAssetLocatorEditor::OnSetObject(const FAssetData& InNewObject, TWeakPtr<IUniversalObjectLocatorCustomization> WeakCustomization)
+void FAssetLocatorEditor::OnSetAsset(const FAssetData& InNewAsset, TWeakPtr<IFragmentEditorHandle> InWeakHandle)
 {
-	TSharedPtr<IUniversalObjectLocatorCustomization> Customization = WeakCustomization.Pin();
-	if (Customization)
+	if (TSharedPtr<IFragmentEditorHandle> Handle = InWeakHandle.Pin())
 	{
 		// Assets are always absolute
-		UObject* Object = InNewObject.FastGetAsset(true);
+		UObject* Object = InNewAsset.FastGetAsset(true);
 
-		FUniversalObjectLocator NewRef(Object, nullptr);
-		Customization->SetValue(MoveTemp(NewRef));
+		FUniversalObjectLocatorFragment NewFragment(FAssetLocatorFragment::FragmentType);
+		FAssetLocatorFragment* Payload = NewFragment.GetPayloadAs(FAssetLocatorFragment::FragmentType);
+		Payload->Path = FTopLevelAssetPath(Object);
+		Handle->SetValue(NewFragment);
 	}
 }
 
+FUniversalObjectLocatorFragment FAssetLocatorEditor::MakeDefaultLocatorFragment() const
+{
+	FUniversalObjectLocatorFragment NewFragment(FAssetLocatorFragment::FragmentType);
+	return NewFragment;
+}
+
 } // namespace UE::UniversalObjectLocator
 
 #undef LOCTEXT_NAMESPACE
diff --git a/Engine/Source/Editor/UniversalObjectLocatorEditor/Private/AssetLocatorEditor.h b/Engine/Source/Editor/UniversalObjectLocatorEditor/Private/AssetLocatorEditor.h
--- a/Engine/Source/Editor/UniversalObjectLocatorEditor/Private/AssetLocatorEditor.h
+++ b/Engine/Source/Editor/UniversalObjectLocatorEditor/Private/AssetLocatorEditor.h
@@ -3,7 +3,7 @@
 #pragma once
 
 #include "CoreTypes.h"
-#include "UniversalObjectLocatorEditor.h"
+#include "UniversalObjectLocatorFragmentEditor.h"
 
 struct FAssetData;
 
@@ -12,23 +12,30 @@ namespace UE::UniversalObjectLocator
 
 class IUniversalObjectLocatorCustomization;
 
-class FAssetLocatorEditor : public ILocatorEditor
+class FAssetLocatorEditor : public ILocatorFragmentEditor
 {
+	ELocatorFragmentEditorType GetLocatorFragmentEditorType() const override;
+
 	bool IsDragSupported(TSharedPtr<FDragDropOperation> DragOperation, UObject* Context) const override;
 
 	UObject* ResolveDragOperation(TSharedPtr<FDragDropOperation> DragOperation, UObject* Context) const override;
 
-	TSharedPtr<SWidget> MakeEditUI(TSharedPtr<IUniversalObjectLocatorCustomization> Customization) override;
+	TSharedPtr<SWidget> MakeEditUI(const FEditUIParameters& InParameters) override;
+
+	FText GetDisplayText(const FUniversalObjectLocatorFragment* InFragment = nullptr) const override;
+
+	FText GetDisplayTooltip(const FUniversalObjectLocatorFragment* InFragment = nullptr) const override;
 
-	FText GetDisplayText() const override;
+	FSlateIcon GetDisplayIcon(const FUniversalObjectLocatorFragment* InFragment = nullptr) const override;
 
-	FText GetDisplayTooltip() const override;
+	UClass* ResolveClass(const FUniversalObjectLocatorFragment& InFragment, UObject* InContext) const override;
 
-	FSlateIcon GetDisplayIcon() const override;
+	FUniversalObjectLocatorFragment MakeDefaultLocatorFragment() const override;
 
 private:
+	FAssetData GetAsset(TWeakPtr<IFragmentEditorHandle> InWeakHandle) const;
 
-	void OnSetObject(const FAssetData& InNewObject, TWeakPtr<IUniversalObjectLocatorCustomization> WeakCustomization);
+	void OnSetAsset(const FAssetData& InNewObject, TWeakPtr<IFragmentEditorHandle> InWeakHandle);
 };
 
 } // namespace UE::UniversalObjectLocator
diff --git a/Engine/Source/Editor/UniversalObjectLocatorEditor/Private/UniversalObjectLocatorCustomization.cpp b/Engine/Source/Editor/UniversalObjectLocatorEditor/Private/UniversalObjectLocatorCustomization.cpp
--- a/Engine/Source/Editor/UniversalObjectLocatorEditor/Private/UniversalObjectLocatorCustomization.cpp
+++ b/Engine/Source/Editor/UniversalObjectLocatorEditor/Private/UniversalObjectLocatorCustomization.cpp
@@ -3,7 +3,7 @@
 #include "UniversalObjectLocatorCustomization.h"
 #include "UniversalObjectLocator.h"
 #include "UniversalObjectLocatorFragmentType.h"
-#include "UniversalObjectLocatorEditor.h"
+#include "UniversalObjectLocatorFragmentEditor.h"
 #include "UniversalObjectLocatorEditorModule.h"
 
 #include "Containers/Array.h"
@@ -24,8 +24,14 @@
 
 #include "SceneOutlinerPublicTypes.h"
 #include "SceneOutlinerModule.h"
+#include "ScopedTransaction.h"
+#include "UniversalObjectLocatorEditorContext.h"
+#include "UniversalObjectLocatorInitializeResult.h"
 
 #include "Async/Async.h"
+#include "String/ParseTokens.h"
+#include "Widgets/Layout/SWrapBox.h"
+#include "Widgets/Views/STileView.h"
 
 
 #define LOCTEXT_NAMESPACE "FUniversalObjectLocatorCustomization"
@@ -33,6 +39,83 @@
 namespace UE::UniversalObjectLocator
 {
 
+// Data representation of a UOL fragment for UI display
+struct FFragmentItem : public IFragmentEditorHandle
+{
+public:
+	// IFragmentEditorHandle interface
+	virtual const FUniversalObjectLocatorFragment& GetFragment() const override
+	{
+		const FUniversalObjectLocator* Value = nullptr;
+		PropertyHandle->EnumerateConstRawData(
+			[&Value](const void* Data, int32 DataIndex, int32 Num)
+			{
+				Value = static_cast<const FUniversalObjectLocator*>(Data);
+				return false;
+			}
+		);
+		check(Value->Fragments.IsValidIndex(FragmentIndex));
+		return Value->Fragments[FragmentIndex];
+	}
+
+	virtual const UClass* GetContextClass() const override
+	{
+		return WeakContextClass.Get();
+	}
+
+	virtual const UClass* GetResolvedClass() const override
+	{
+		return WeakResolvedClass.Get();
+	}
+	
+	virtual void SetValue(const FUniversalObjectLocatorFragment& InNewValue) override
+	{
+		FScopedTransaction Transaction(LOCTEXT("EditLocatorFragmentTransaction", "Edit Locator Fragment"));
+
+		PropertyHandle->NotifyPreChange();
+
+		PropertyHandle->EnumerateRawData([this, &InNewValue](void* RawData, int32 Index, int32 Num){
+			FUniversalObjectLocator* Reference = static_cast<FUniversalObjectLocator*>(RawData);
+			check(Reference->Fragments.IsValidIndex(FragmentIndex));
+			Reference->Fragments[FragmentIndex] = InNewValue;
+			return true;
+		});
+
+		if(TSharedPtr<FUniversalObjectLocatorCustomization> Customization = WeakCustomization.Pin())
+		{
+			Customization->TrimAbsoluteFragments();
+		}
+
+		PropertyHandle->NotifyPostChange(EPropertyChangeType::ValueSet);
+		PropertyHandle->NotifyFinishedChangingProperties();
+	}
+
+	// Fragment index in the overall fragment chain
+	int32 FragmentIndex = INDEX_NONE;
+	// The type of the fragment
+	FFragmentTypeHandle FragmentType;
+	// Locator editor type
+	FName LocatorEditorType;
+	// Locator editor
+	TSharedPtr<ILocatorFragmentEditor> LocatorEditor;
+	// Context class, if any
+	TWeakObjectPtr<const UClass> WeakContextClass;
+	// Resolved class, if any
+	TWeakObjectPtr<const UClass> WeakResolvedClass;
+	// Cached display text
+	TOptional<FText> DisplayText;
+	// Cached display tooltip text
+	TOptional<FText> TooltipText;
+	// Cached display icon
+	TOptional<const FSlateBrush*> Icon;
+	// Property handle to the containing FUniversalObjectLocator
+	TSharedPtr<IPropertyHandle> PropertyHandle;
+	// Owning customization
+	TWeakPtr<FUniversalObjectLocatorCustomization> WeakCustomization; 
+	// Whether this is the last fragment
+	bool bIsTail = false;
+};
+
 TSharedRef<IPropertyTypeCustomization> FUniversalObjectLocatorCustomization::MakeInstance()
 {
 	return MakeShared<FUniversalObjectLocatorCustomization>();
@@ -48,96 +131,298 @@ void FUniversalObjectLocatorCustomization::CustomizeHeader(TSharedRef<IPropertyH
 		InitOptions.bShowHeaderRow = false;
 	}
 
-	bool bDisplayThumbnail = true;
+	FUniversalObjectLocatorEditorModule& EditorModule = FModuleManager::Get().LoadModuleChecked<FUniversalObjectLocatorEditorModule>("UniversalObjectLocatorEditor");
 
-	TArray<FString> AllowedTypes;
-	if (PropertyHandle->HasMetaData("AllowedLocators"))
+	if (PropertyHandle->HasMetaData("LocatorContext"))
 	{
-		PropertyHandle->GetMetaData("AllowedLocators").ParseIntoArray(AllowedTypes, TEXT(","));
-	}
+		const FString& EditorContext = PropertyHandle->GetMetaData("LocatorContext");
+		FName EditorContextName(*EditorContext);
+		if(const TSharedPtr<ILocatorFragmentEditorContext> FoundContext = EditorModule.LocatorEditorContexts.FindRef(EditorContextName))
+		{
+			if(WeakContext.Get() == nullptr)
+			{
+				WeakContext = FoundContext->GetContext(PropertyHandle.ToSharedRef().Get());
+				if(UObject* Context = WeakContext.Get())
+				{
+					if(Context->IsA<UClass>())
+					{
+						WeakContextClass = Cast<UClass>(Context);
+					}
+					else
+					{
+						WeakContextClass = Context->GetClass();
+					}
+				}
+			}
 
-	FUniversalObjectLocatorEditorModule& Module = FModuleManager::Get().LoadModuleChecked<FUniversalObjectLocatorEditorModule>("UniversalObjectLocatorEditor");
-	for (TPair<FName, TSharedPtr<ILocatorEditor>> Pair : Module.LocatorEditors)
+			for (TPair<FName, TSharedPtr<ILocatorFragmentEditor>> Pair : EditorModule.LocatorEditors)
+			{
+				if(Pair.Value->IsAllowedInContext(EditorContextName) && FoundContext->IsFragmentAllowed(Pair.Key))
+				{
+					ApplicableLocators.Add(Pair.Key, Pair.Value);
+				}
+			}
+		}
+	}
+	else
 	{
-		if (AllowedTypes.Num() == 0 || AllowedTypes.Contains(Pair.Key.ToString()))
+		TArray<FString> AllowedTypes;
+		if (PropertyHandle->HasMetaData("AllowedLocators"))
 		{
-			ApplicableLocators.Add(Pair.Key, Pair.Value);
+			const FString& AllowedLocatorsNames = PropertyHandle->GetMetaData("AllowedLocators");
+			UE::String::ParseTokens(AllowedLocatorsNames, TEXT(','), [&AllowedTypes](FStringView InToken)
+			{
+				AllowedTypes.Add(FString(InToken));
+			}, UE::String::EParseTokensOptions::SkipEmpty | UE::String::EParseTokensOptions::Trim);
 		}
+
+		for (TPair<FName, TSharedPtr<ILocatorFragmentEditor>> Pair : EditorModule.LocatorEditors)
+		{
+			if ((AllowedTypes.Num() == 0 || AllowedTypes.Contains(Pair.Key.ToString())) && Pair.Value->IsAllowedInContext(NAME_None))
+			{
+				ApplicableLocators.Add(Pair.Key, Pair.Value);
+			}
+		}
+		
+		TArray<UObject*> OuterObjects;
+		PropertyHandle->GetOuterObjects(OuterObjects);
+		WeakContext = OuterObjects.Num() ? OuterObjects[0] : nullptr;
+		WeakContextClass = OuterObjects.Num() ? OuterObjects[0]->GetClass() : UObject::StaticClass();
 	}
 
-	Content = SNew(SBox);
 	HeaderRow
 	.NameContent()
 	[
-		SNew(SHorizontalBox)
-
-		+ SHorizontalBox::Slot()
-		.AutoWidth()
-		[
-			StructPropertyHandle->CreatePropertyNameWidget()
-		]
-
-		// Type combo box
-		+ SHorizontalBox::Slot()
-		.Padding(FMargin(5.f, 0.f, 0.f, 0.f))
-		.AutoWidth()
-		[
-			SNew(SComboButton)
-			.Visibility(ApplicableLocators.Num() > 1 ? EVisibility::Visible : EVisibility::Collapsed)
-			.OnGetMenuContent(this, &FUniversalObjectLocatorCustomization::GetUserExposedFragmentTypeList)
-			.ButtonContent()
-			[
-				SNew(STextBlock)
-				.Text(this, &FUniversalObjectLocatorCustomization::GetCurrentFragmentTypeText)
-			]
-		]
+		StructPropertyHandle->CreatePropertyNameWidget()
 	]
 	.ValueContent()
-	.MaxDesiredWidth(TOptional<float>())
+	.MaxDesiredWidth(400.0f)
 	[
-		SNew(SDropTarget)
-		.OnAllowDrop(this, &FUniversalObjectLocatorCustomization::HandleIsDragAllowed)
-		.OnDropped(this, &FUniversalObjectLocatorCustomization::HandleDrop)
+		SNew(SBorder)
+		.ToolTipText_Lambda([this]()
+		{
+			FTextBuilder TextBuilder;
+			TextBuilder.AppendLine(PropertyHandle->GetProperty()->GetDisplayNameText());
+
+			if(GetCachedData().PropertyValue.IsSet())
+			{
+				TStringBuilder<256> StringBuilder;
+				GetCachedData().PropertyValue.GetValue().ToString(StringBuilder);
+				TextBuilder.AppendLine(FText::FromStringView(StringBuilder));
+			}
+			else
+			{
+				TArray<UObject*> Objects;
+				PropertyHandle->GetOuterObjects(Objects);
+				if(Objects.Num() > 1)
+				{
+					TextBuilder.AppendLine(LOCTEXT("MultipleValues", "Multiple Values"));
+				}
+				else
+				{
+					TextBuilder.AppendLine(LOCTEXT("NoValues", "None"));
+				}
+			}
+
+			return TextBuilder.ToText();
+		})
+		.BorderImage(&FAppStyle::Get().GetWidgetStyle<FComboButtonStyle>("ComboButton").ButtonStyle.Normal)
+		.Padding(2.0f, 2.0f, 2.0f, 2.0f)
 		[
-			Content.ToSharedRef()
+			SAssignNew(WrapBox, SWrapBox)
+			.PreferredSize(400.0f)
 		]
 	];
 
+	PropertyHandle->SetOnPropertyValueChanged(FSimpleDelegate::CreateLambda([this]()
+	{
+		Rebuild();
+	}));
+
 	Rebuild();
 }
 
 void FUniversalObjectLocatorCustomization::CustomizeChildren(TSharedRef<IPropertyHandle> StructPropertyHandle, IDetailChildrenBuilder& StructBuilder, IPropertyTypeCustomizationUtils& StructCustomizationUtils)
 {
-	// Allow the custom locator editor to customize children if desired
-	FUniversalObjectLocatorEditorModule& Module = FModuleManager::Get().LoadModuleChecked<FUniversalObjectLocatorEditorModule>("UniversalObjectLocatorEditor");
-	TSharedPtr<ILocatorEditor> LocatorEditor = ApplicableLocators.FindRef(GetCachedData().LocatorEditorType);
-	if (LocatorEditor)
-	{
-		LocatorEditor->CustomizeChildren(StructPropertyHandle, StructBuilder, StructCustomizationUtils);
-	}
 }
 
-void FUniversalObjectLocatorCustomization::Rebuild() const
+void FUniversalObjectLocatorCustomization::Rebuild()
 {
-	FUniversalObjectLocatorEditorModule& Module = FModuleManager::Get().LoadModuleChecked<FUniversalObjectLocatorEditorModule>("UniversalObjectLocatorEditor");
-	TSharedPtr<ILocatorEditor> LocatorEditor = ApplicableLocators.FindRef(GetCachedData().LocatorEditorType);
-	if (LocatorEditor)
+	Fragments.Reset();
+
+	const FUniversalObjectLocator* CommonValue = GetCommonPropertyValue();
+	if(CommonValue)
 	{
-		TSharedPtr<SWidget> EditUI = LocatorEditor->MakeEditUI(SharedThis(const_cast<FUniversalObjectLocatorCustomization*>(this)));
-		if (EditUI)
+		// Keep track of current context & class and try to resolve as we traverse the path
+		UClass* CurrentContextClass = WeakContextClass.Get(); 
+		for(int32 FragmentIndex = 0; FragmentIndex < CommonValue->Fragments.Num(); ++FragmentIndex)
 		{
-			TSharedPtr<SBox> BoxWidget = Content;
+			const FUniversalObjectLocatorFragment& Fragment = CommonValue->Fragments[FragmentIndex];
+			TSharedRef<FFragmentItem> NewItem = MakeShared<FFragmentItem>();
+			NewItem->FragmentIndex = FragmentIndex;
+			NewItem->FragmentType = Fragment.GetFragmentTypeHandle();
+			NewItem->LocatorEditorType = Fragment.GetFragmentType()->PrimaryEditorType;
+			NewItem->LocatorEditor = ApplicableLocators.FindRef(Fragment.GetFragmentType()->PrimaryEditorType);
+			NewItem->PropertyHandle = PropertyHandle;
+			NewItem->WeakCustomization = SharedThis(this);
+			NewItem->bIsTail = FragmentIndex == CommonValue->Fragments.Num() - 1;
 			
-			// Delay reconstruction until next frame since this can be called from inside the previous EditUI,
-			// leading to destruction of the UI while it is still running
-			AsyncTask(ENamedThreads::GameThread, [BoxWidget, EditUI, this] {
-				BoxWidget->SetContent(EditUI.ToSharedRef());
-			});
-			return;
+			NewItem->WeakContextClass = CurrentContextClass;
+			if(NewItem->LocatorEditor.IsValid())
+			{
+				// Make a partial Uol to try to resolve a valid class at this fragment
+				FUniversalObjectLocator PartialUol;
+				for(int32 PartialFragmentIndex = 0; PartialFragmentIndex < FragmentIndex; ++PartialFragmentIndex)
+				{
+					PartialUol.Fragments.Add(CommonValue->Fragments[PartialFragmentIndex]);
+				}
+
+				FResolveResult Result = PartialUol.Resolve(FResolveParams(GetContext()));
+				CurrentContextClass = NewItem->LocatorEditor->ResolveClass(Fragment, Result.SyncGet().Object);
+				NewItem->WeakResolvedClass = CurrentContextClass;
+			}
+			else
+			{
+				CurrentContextClass = nullptr;
+			}
+
+			Fragments.Add(NewItem);
 		}
 	}
 
-	Content->SetContent(SNullWidget::NullWidget);
+	// Clear the cache to force cache to rebuild
+	CachedData.PropertyValue.Reset();
+
+	WrapBox->ClearChildren();
+
+	for(TSharedRef<FFragmentItem> FragmentItem : Fragments)
+	{
+		TWeakPtr<FFragmentItem> WeakFragmentItem = FragmentItem;
+
+		WrapBox->AddSlot()
+		.Padding(0.0f, 0.0f, 1.0f, 0.0f)
+		[
+			SNew(SDropTarget)
+			.OnAllowDrop(this, &FUniversalObjectLocatorCustomization::HandleIsDragAllowed, WeakFragmentItem)
+			.OnDropped(this, &FUniversalObjectLocatorCustomization::HandleDrop, WeakFragmentItem)
+			[
+				SNew(SHorizontalBox)
+				+SHorizontalBox::Slot()
+				.AutoWidth()
+				[
+					SNew(SComboButton)
+					.ComboButtonStyle(&FAppStyle::Get().GetWidgetStyle<FComboButtonStyle>("SimpleComboButton"))
+					.ToolTipText(LOCTEXT("SeperatorTooltip", "Edit this fragment"))
+					.HasDownArrow(false)
+					.OnGetMenuContent(this, &FUniversalObjectLocatorCustomization::GetUserExposedFragmentTypeList, WeakFragmentItem)
+					.ButtonContent()
+					[
+						SNew(SImage)
+						.Image(FAppStyle::Get().GetBrush("Icons.ChevronRight"))
+					]
+				]
+				+SHorizontalBox::Slot()
+				.AutoWidth()
+				[
+					SNew(SComboButton)
+					.ComboButtonStyle(&FAppStyle::Get().GetWidgetStyle<FComboButtonStyle>("SimpleComboButton"))
+					.ToolTipText(this, &FUniversalObjectLocatorCustomization::GetFragmentTooltipText, WeakFragmentItem)
+					.HasDownArrow(false)
+					.OnGetMenuContent(this, &FUniversalObjectLocatorCustomization::GetFragmentTypeWidget, WeakFragmentItem)
+					.ButtonContent()
+					[
+						SNew(SBox)
+						.MinDesiredHeight(16.0f)
+						.VAlign(VAlign_Center)
+						[
+							SNew(SHorizontalBox)
+							+SHorizontalBox::Slot()
+							.AutoWidth()
+							.VAlign(VAlign_Center)
+							[
+								SNew(SBox)
+								.MaxDesiredWidth(16.0f)
+								.MaxDesiredHeight(16.0f)
+								[
+									SNew(SImage)
+									.Image(this, &FUniversalObjectLocatorCustomization::GetFragmentIcon, WeakFragmentItem)
+								]
+							]
+							+SHorizontalBox::Slot()
+							.AutoWidth()
+							.VAlign(VAlign_Center)
+							.Padding(2.0f, 0.0f, 0.0f, 0.0f)
+							[
+								SNew(STextBlock)
+								.TextStyle(&FAppStyle::Get().GetWidgetStyle<FTextBlockStyle>("SmallText"))
+								.Text(this, &FUniversalObjectLocatorCustomization::GetFragmentText, WeakFragmentItem)
+							]
+						]
+					]
+				]
+			]
+		];
+	}
+
+	WrapBox->AddSlot()
+	[
+		SNew(SDropTarget)
+		.OnAllowDrop(this, &FUniversalObjectLocatorCustomization::HandleIsDragAllowed, TWeakPtr<FFragmentItem>())
+		.OnDropped(this, &FUniversalObjectLocatorCustomization::HandleDrop, TWeakPtr<FFragmentItem>())
+		[
+			SNew(SComboButton)
+			.ComboButtonStyle(&FAppStyle::Get().GetWidgetStyle<FComboButtonStyle>("SimpleComboButton"))
+			.ToolTipText(LOCTEXT("SeperatorAddFragmentTooltip", "Add a fragment to this locator"))
+			.HasDownArrow(false)
+			.OnGetMenuContent(this, &FUniversalObjectLocatorCustomization::GetUserExposedFragmentTypeList, TWeakPtr<FFragmentItem>())
+			.ButtonContent()
+			[
+				SNew(SImage)
+				.Image(FAppStyle::Get().GetBrush("Icons.ChevronRight"))
+			]
+		]
+	];
+
+	TArray<UObject*> Objects;
+	PropertyHandle->GetOuterObjects(Objects);
+	if(CommonValue == nullptr && Objects.Num() > 1)
+	{
+		WrapBox->AddSlot()
+		.VAlign(VAlign_Center)
+		[
+			SNew(STextBlock)
+			.Text(LOCTEXT("MultipleValues", "Multiple Values"))
+			.TextStyle(&FAppStyle::Get().GetWidgetStyle<FTextBlockStyle>("SmallText"))
+		];
+	}
+}
+
+void FUniversalObjectLocatorCustomization::TrimAbsoluteFragments()
+{
+	// When edited, certain configurations are invalid (like absolute fragments in the middle of a locator's path),
+	// so we collapse the edited fragments here
+	PropertyHandle->EnumerateRawData([this](void* RawData, int32 Index, int32 Num){
+		FUniversalObjectLocator* Reference = static_cast<FUniversalObjectLocator*>(RawData);
+
+		for(int32 FragmentIndex = 0; FragmentIndex < Reference->Fragments.Num(); ++FragmentIndex)
+		{
+			const FUniversalObjectLocatorFragment& Fragment = Reference->Fragments[FragmentIndex];
+			if(TSharedPtr<ILocatorFragmentEditor> LocatorEditor = ApplicableLocators.FindRef(Fragment.GetFragmentType()->PrimaryEditorType))
+			{
+				if(LocatorEditor->GetLocatorFragmentEditorType() == ELocatorFragmentEditorType::Absolute)
+				{
+					if(FragmentIndex != 0)
+					{
+						// Absolute fragments always trim fragments before them
+						Reference->Fragments.RemoveAt(0, FragmentIndex);
+						break;
+					}
+				}
+			}
+		}
+
+		return true;
+	});
 }
 
 UObject* FUniversalObjectLocatorCustomization::GetSingleObject() const
@@ -145,8 +430,10 @@ UObject* FUniversalObjectLocatorCustomization::GetSingleObject() const
 	return GetCachedData().WeakObject.Get();
 }
 
-void FUniversalObjectLocatorCustomization::SetValue(FUniversalObjectLocator&& InNewValue) const
+void FUniversalObjectLocatorCustomization::SetValue(FUniversalObjectLocator&& InNewValue)
 {
+	FScopedTransaction Transaction(LOCTEXT("EditLocatorTransaction", "Edit Locator"));
+
 	PropertyHandle->NotifyPreChange();
 
 	PropertyHandle->EnumerateRawData([&InNewValue](void* RawData, int32 Index, int32 Num){
@@ -162,13 +449,10 @@ void FUniversalObjectLocatorCustomization::SetValue(FUniversalObjectLocator&& In
 		return true;
 	});
 
+	TrimAbsoluteFragments();
+
 	PropertyHandle->NotifyPostChange(EPropertyChangeType::ValueSet);
 	PropertyHandle->NotifyFinishedChangingProperties();
-
-	// Clear the cache to force cache to rebuild
-	CachedData.PropertyValue.Reset();
-
-	Rebuild();
 }
 
 FString FUniversalObjectLocatorCustomization::GetPathToObject() const
@@ -181,11 +465,11 @@ TSharedPtr<IPropertyHandle> FUniversalObjectLocatorCustomization::GetProperty()
 	return PropertyHandle;
 }
 
-TSharedRef<SWidget> FUniversalObjectLocatorCustomization::GetUserExposedFragmentTypeList()
+TSharedRef<SWidget> FUniversalObjectLocatorCustomization::GetUserExposedFragmentTypeList(TWeakPtr<FFragmentItem> InWeakFragmentItem)
 {
 	struct FMenuData
 	{
-		TSharedPtr<ILocatorEditor> LocatorEditor;
+		TSharedPtr<ILocatorFragmentEditor> LocatorEditor;
 		FText DisplayText;
 		FText DisplayTooltip;
 		FSlateIcon DisplayIcon;
@@ -194,7 +478,7 @@ TSharedRef<SWidget> FUniversalObjectLocatorCustomization::GetUserExposedFragment
 
 	TArray<FMenuData> MenuData;
 
-	for (TPair<FName, TSharedPtr<ILocatorEditor>> Pair : ApplicableLocators)
+	for (TPair<FName, TSharedPtr<ILocatorFragmentEditor>> Pair : ApplicableLocators)
 	{
 		MenuData.Add(FMenuData{
 			Pair.Value,
@@ -214,140 +498,413 @@ TSharedRef<SWidget> FUniversalObjectLocatorCustomization::GetUserExposedFragment
 	Algo::SortBy(MenuData, &FMenuData::DisplayText, FText::FSortPredicate(ETextComparisonLevel::Default));
 
 	const bool bCloseAfterSelection = true;
-	FMenuBuilder MenuBuilder(bCloseAfterSelection, nullptr);
+	const bool bCloseSelfOnly = true;
+	FMenuBuilder MenuBuilder(bCloseAfterSelection, nullptr, nullptr, bCloseSelfOnly);
+
+	TSharedPtr<FFragmentItem> FragmentItem = InWeakFragmentItem.Pin();
+	if(FragmentItem.IsValid() && FragmentItem->FragmentIndex == 0)
+	{
+		MenuBuilder.AddMenuEntry(
+			LOCTEXT("ClearFragmentsLabel", "Clear Fragments"),
+			LOCTEXT("ClearFragmentsTooltip", "Remove all locator fragments"),
+			FSlateIcon(),
+			FUIAction(
+				FExecuteAction::CreateSP(this, &FUniversalObjectLocatorCustomization::ClearFragments)
+			),
+			NAME_None,
+			EUserInterfaceActionType::Button
+		);
+	}
+
+	const bool bIsTail = !FragmentItem.IsValid();
+	if(!bIsTail)
+	{
+		MenuBuilder.AddMenuEntry(
+			LOCTEXT("RemoveFragmentLabel", "Remove Fragment"),
+			LOCTEXT("RemoveFragmentTooltip", "Remove this locator fragment"),
+			FSlateIcon(),
+			FUIAction(
+				FExecuteAction::CreateSP(this, &FUniversalObjectLocatorCustomization::RemoveFragment, InWeakFragmentItem)
+			),
+			NAME_None,
+			EUserInterfaceActionType::Button
+		);
+	}
 
 	FCanExecuteAction AlwaysExecute = FCanExecuteAction::CreateLambda([]{ return true; });
 
-	MenuBuilder.BeginSection(NAME_None, LOCTEXT("FragmentTypeTypeHeader", "Change this locator's to reference a different type:"));
+	const FText SectionLabel = bIsTail ?
+		LOCTEXT("AddFragmentTypeHeader", "Add a new locator fragment") :
+		LOCTEXT("ChangeFragmentTypeHeader", "Change this locator fragment type");
+	MenuBuilder.BeginSection(NAME_None, SectionLabel);
 	for (const FMenuData& Item : MenuData)
 	{
+		// Skip absolute fragments after the first fragment
+		const bool bIsAfterFirstFragment = (FragmentItem.IsValid() && FragmentItem->FragmentIndex != 0) || (bIsTail && Fragments.Num() != 0); 
+		if(bIsAfterFirstFragment && Item.LocatorEditor->GetLocatorFragmentEditorType() == ELocatorFragmentEditorType::Absolute)
+		{
+			continue;
+		}
+		
 		MenuBuilder.AddMenuEntry(
 			Item.DisplayText,
 			Item.DisplayTooltip,
 			Item.DisplayIcon,
 			FUIAction(
-				FExecuteAction::CreateSP(this, &FUniversalObjectLocatorCustomization::ChangeEditorType, Item.LocatorEditorType),
+				FExecuteAction::CreateSP(this, &FUniversalObjectLocatorCustomization::ChangeEditorType, TWeakPtr<ILocatorFragmentEditor>(Item.LocatorEditor), InWeakFragmentItem),
 				AlwaysExecute,
-				FIsActionChecked::CreateSP(this, &FUniversalObjectLocatorCustomization::CompareCurrentEditorType, Item.LocatorEditorType)
+				FIsActionChecked::CreateSP(this, &FUniversalObjectLocatorCustomization::CompareCurrentEditorType, TWeakPtr<ILocatorFragmentEditor>(Item.LocatorEditor), InWeakFragmentItem)
 			),
 			NAME_None,
-			EUserInterfaceActionType::RadioButton
+			bIsTail ? EUserInterfaceActionType::Button : EUserInterfaceActionType::RadioButton
 		);
 	}
 	MenuBuilder.EndSection();
 
-
 	return MenuBuilder.MakeWidget();
 }
 
-void FUniversalObjectLocatorCustomization::ChangeEditorType(FName InNewLocatorEditorType)
+void FUniversalObjectLocatorCustomization::ChangeEditorType(TWeakPtr<ILocatorFragmentEditor> InNewLocatorEditor, TWeakPtr<FFragmentItem> InWeakFragmentItem)
 {
+	TSharedPtr<FFragmentItem> FragmentItem = InWeakFragmentItem.Pin();
+
+	const FText TransactionText = !FragmentItem.IsValid() ?
+		LOCTEXT("AddFragmentTransaction", "Add Locator Fragment") :
+		LOCTEXT("ChangeLocatorFragmentTypeTransaction", "Change Locator Fragment Type");
+	FScopedTransaction Transaction(TransactionText);
+
 	PropertyHandle->NotifyPreChange();
 
 	PropertyHandle->EnumerateRawData(
-		[InNewLocatorEditorType, this](void* Data, int32 DataIndex, int32 Num)
+		[InNewLocatorEditor, this, &FragmentItem](void* Data, int32 DataIndex, int32 Num)
 		{
-			FUniversalObjectLocator* Ref = static_cast<FUniversalObjectLocator*>(Data);
-			Ref->Reset();
+			if (TSharedPtr<ILocatorFragmentEditor> LocatorEditor = InNewLocatorEditor.Pin())
+			{
+				FUniversalObjectLocator* Ref = static_cast<FUniversalObjectLocator*>(Data);
+				if(FragmentItem.IsValid() && Ref->Fragments.IsValidIndex(FragmentItem->FragmentIndex))
+				{
+					// Trim the remaining fragments beyond this one
+					Ref->Fragments.RemoveAt(FragmentItem->FragmentIndex, Ref->Fragments.Num() - FragmentItem->FragmentIndex);
+					Ref->AddFragment(LocatorEditor->MakeDefaultLocatorFragment());
+				}
+				else
+				{
+					Ref->AddFragment(LocatorEditor->MakeDefaultLocatorFragment());
+				}
+			}
 
-			TSharedPtr<ILocatorEditor> LocatorEditor = ApplicableLocators.FindRef(InNewLocatorEditorType);
-			if (LocatorEditor)
+			return true;
+		}
+	);
+
+	TrimAbsoluteFragments();
+
+	PropertyHandle->NotifyPostChange(EPropertyChangeType::ValueSet);
+	PropertyHandle->NotifyFinishedChangingProperties();
+}
+
+void FUniversalObjectLocatorCustomization::RemoveFragment(TWeakPtr<FFragmentItem> InWeakFragmentItem)
+{
+	TSharedPtr<FFragmentItem> FragmentItem = InWeakFragmentItem.Pin();
+	if(!FragmentItem.IsValid())
+	{
+		return;
+	}
+
+	FScopedTransaction Transaction(LOCTEXT("RemoveLocatorFragmentTransaction", "Remove Locator Fragment"));
+	
+	PropertyHandle->NotifyPreChange();
+
+	PropertyHandle->EnumerateRawData(
+		[this, &FragmentItem](void* Data, int32 DataIndex, int32 Num)
+		{
+			FUniversalObjectLocator* Ref = static_cast<FUniversalObjectLocator*>(Data);
+			if(FragmentItem.IsValid() && Ref->Fragments.IsValidIndex(FragmentItem->FragmentIndex))
 			{
-				*Ref = LocatorEditor->MakeDefaultLocator();
+				Ref->Fragments.RemoveAt(FragmentItem->FragmentIndex);
 			}
 			return true;
 		}
 	);
 
-	CachedData.LocatorEditorType = InNewLocatorEditorType;
+	TrimAbsoluteFragments();
 
 	PropertyHandle->NotifyPostChange(EPropertyChangeType::ValueSet);
 	PropertyHandle->NotifyFinishedChangingProperties();
+}
 
-	PropertyUtilities->ForceRefresh();
+void FUniversalObjectLocatorCustomization::ClearFragments()
+{
+	FScopedTransaction Transaction(LOCTEXT("ClearLocatorFragmentsTransaction", "Clear Locator Fragments"));
+	
+	PropertyHandle->NotifyPreChange();
 
-	Rebuild();
+	PropertyHandle->EnumerateRawData(
+		[this](void* Data, int32 DataIndex, int32 Num)
+		{
+			FUniversalObjectLocator* Ref = static_cast<FUniversalObjectLocator*>(Data);
+			Ref->Reset();
+			return true;
+		}
+	);
 
+	PropertyHandle->NotifyPostChange(EPropertyChangeType::ValueSet);
+	PropertyHandle->NotifyFinishedChangingProperties();
 }
 
-bool FUniversalObjectLocatorCustomization::CompareCurrentEditorType(FName InLocatorEditorType) const
+bool FUniversalObjectLocatorCustomization::CompareCurrentEditorType(TWeakPtr<ILocatorFragmentEditor> InNewLocatorEditor, TWeakPtr<FFragmentItem> InWeakFragmentItem) const
 {
-	return CachedData.LocatorEditorType == InLocatorEditorType;
+	TSharedPtr<FFragmentItem> FragmentItem = InWeakFragmentItem.Pin();
+	if(!FragmentItem.IsValid())
+	{
+		return false;
+	}
+	
+	return FragmentItem->LocatorEditor == InNewLocatorEditor.Pin();
 }
 
-FText FUniversalObjectLocatorCustomization::GetCurrentFragmentTypeText() const
+FText FUniversalObjectLocatorCustomization::GetFragmentText(TWeakPtr<FFragmentItem> InWeakFragmentItem) const
 {
-	if (CachedData.LocatorEditorType != NAME_None)
+	const FText NoneText = LOCTEXT("NoValues", "None");
+	TSharedPtr<FFragmentItem> FragmentItem = InWeakFragmentItem.Pin();
+	if(!FragmentItem.IsValid())
 	{
-		return FText::FromName(CachedData.LocatorEditorType);
+		return NoneText;
 	}
 
-	TArray<void*> RawData;
-	PropertyHandle->AccessRawData(RawData);
+	if(FragmentItem->DisplayText.IsSet())
+	{
+		return FragmentItem->DisplayText.GetValue();
+	}
+	
+	if (FragmentItem->LocatorEditor.IsValid())
+	{
+		FragmentItem->DisplayText = FragmentItem->LocatorEditor->GetDisplayText(&FragmentItem->GetFragment());
+		return FragmentItem->DisplayText.GetValue();
+	}
 
-	FText NoneText = LOCTEXT("NoValues", "None");
-	if (RawData.Num() == 0)
+	if(FragmentItem->LocatorEditorType != NAME_None)
+	{
+		FragmentItem->DisplayText = FText::FromName(FragmentItem->LocatorEditorType);
+		return FragmentItem->DisplayText.GetValue();
+	}
+
+	return NoneText;
+}
+
+FText FUniversalObjectLocatorCustomization::GetFragmentTooltipText(TWeakPtr<FFragmentItem> InWeakFragmentItem) const
+{
+	const FText NoneText = LOCTEXT("NoValues", "None");
+	TSharedPtr<FFragmentItem> FragmentItem = InWeakFragmentItem.Pin();
+	if(!FragmentItem.IsValid())
 	{
 		return NoneText;
 	}
 
-	const FFragmentType* CommonFragmentType = nullptr;
-	for (void* Ptr : RawData)
+	if(FragmentItem->TooltipText.IsSet())
 	{
-		FUniversalObjectLocator* Reference    = static_cast<FUniversalObjectLocator*>(RawData[0]);
-		const FFragmentType*     FragmentType = Reference->GetLastFragmentType();
+		return FragmentItem->TooltipText.GetValue();
+	}
 
-		if (CommonFragmentType && FragmentType != CommonFragmentType)
-		{
-			return LOCTEXT("MultipleTypes", "<multiple>");
-		}
-		CommonFragmentType = FragmentType;
+	if (FragmentItem->LocatorEditor.IsValid())
+	{
+		FragmentItem->TooltipText = FragmentItem->LocatorEditor->GetDisplayTooltip(&FragmentItem->GetFragment());
+		return FragmentItem->TooltipText.GetValue();
 	}
-	return CommonFragmentType ? CommonFragmentType->DisplayText : NoneText;
+
+	if (FragmentItem->LocatorEditorType != NAME_None)
+	{
+		FragmentItem->TooltipText = FText::FromName(FragmentItem->LocatorEditorType);
+		return FragmentItem->TooltipText.GetValue();
+	}
+
+	return NoneText;
 }
 
-bool FUniversalObjectLocatorCustomization::HandleIsDragAllowed(TSharedPtr<FDragDropOperation> InDragOperation)
+const FSlateBrush* FUniversalObjectLocatorCustomization::GetFragmentIcon(TWeakPtr<FFragmentItem> InWeakFragmentItem) const
 {
-	for (TPair<FName, TSharedPtr<ILocatorEditor>> Pair : ApplicableLocators)
+	TSharedPtr<FFragmentItem> FragmentItem = InWeakFragmentItem.Pin();
+	if(!FragmentItem.IsValid())
+	{
+		return nullptr;
+	}
+
+	if(FragmentItem->Icon.IsSet())
 	{
-		if (Pair.Value->IsDragSupported(InDragOperation, nullptr))
+		return FragmentItem->Icon.GetValue();
+	}
+	
+	if (FragmentItem->LocatorEditorType != NAME_None)
+	{
+		if(const TSharedPtr<ILocatorFragmentEditor>* FoundLocator = ApplicableLocators.Find(FragmentItem->LocatorEditorType))
 		{
-			return true;
+			check(FoundLocator->IsValid());
+			FragmentItem->Icon = FoundLocator->Get()->GetDisplayIcon(&FragmentItem->GetFragment()).GetIcon();
+			return FragmentItem->Icon.GetValue();
 		}
 	}
 
-	return false;
+	return nullptr;
 }
 
-FReply FUniversalObjectLocatorCustomization::HandleDrop(const FGeometry& Geometry, const FDragDropEvent& DragDropEvent)
+TSharedRef<SWidget> FUniversalObjectLocatorCustomization::GetFragmentTypeWidget(TWeakPtr<FFragmentItem> InWeakFragmentItem)
 {
-	TSharedPtr<FDragDropOperation> DragOperation = DragDropEvent.GetOperation();
+	TSharedPtr<FFragmentItem> FragmentItem = InWeakFragmentItem.Pin();
+	if(!FragmentItem.IsValid())
+	{
+		return SNullWidget::NullWidget;
+	}
+	
+	FUniversalObjectLocatorEditorModule& Module = FModuleManager::Get().LoadModuleChecked<FUniversalObjectLocatorEditorModule>("UniversalObjectLocatorEditor");
+	TSharedPtr<ILocatorFragmentEditor> LocatorEditor = ApplicableLocators.FindRef(FragmentItem->LocatorEditorType);
+	if (LocatorEditor)
+	{
+		TSharedPtr<SWidget> EditUI = LocatorEditor->MakeEditUI({ SharedThis(const_cast<FUniversalObjectLocatorCustomization*>(this)), FragmentItem.ToSharedRef() });
+		if (EditUI)
+		{
+			return
+				SNew(SBox)
+				.Padding(5.0f)
+				[
+					EditUI.ToSharedRef()
+				];
+		}
+	}
 
-	UObject* ResolvedObject = nullptr;
+	return SNullWidget::NullWidget;
+}
 
-	for (TPair<FName, TSharedPtr<ILocatorEditor>> Pair : ApplicableLocators)
+bool FUniversalObjectLocatorCustomization::HandleIsDragAllowed(TSharedPtr<FDragDropOperation> InDragOperation, TWeakPtr<FFragmentItem> InWeakFragmentItem)
+{
+	if(TSharedPtr<FFragmentItem> FragmentItem = InWeakFragmentItem.Pin())
 	{
-		ResolvedObject = Pair.Value->ResolveDragOperation(DragOperation, nullptr);
-		if (ResolvedObject)
+		// Dragging onto an existing fragment
+		if(FragmentItem->LocatorEditor.IsValid() && FragmentItem->LocatorEditor->IsDragSupported(InDragOperation, nullptr))
 		{
-			break;
+			return true;
 		}
 	}
-
-	if (!ResolvedObject)
+	else
 	{
-		return FReply::Unhandled();
+		// Dragging onto the tail
+		for (TPair<FName, TSharedPtr<ILocatorFragmentEditor>> Pair : ApplicableLocators)
+		{
+			if (Pair.Value->IsDragSupported(InDragOperation, nullptr))
+			{
+				return true;
+			}
+		}
 	}
+	
+	return false;
+}
 
-	TArray<void*> RawData;
-	PropertyHandle->AccessRawData(RawData);
-	for (void* Ptr : RawData)
+FReply FUniversalObjectLocatorCustomization::HandleDrop(const FGeometry& Geometry, const FDragDropEvent& DragDropEvent, TWeakPtr<FFragmentItem> InWeakFragmentItem)
+{
+	TSharedPtr<FDragDropOperation> DragOperation = DragDropEvent.GetOperation();
+	
+	if(TSharedPtr<FFragmentItem> FragmentItem = InWeakFragmentItem.Pin())
 	{
-		FUniversalObjectLocator* Reference = static_cast<FUniversalObjectLocator*>(Ptr);
-		Reference->Reset(ResolvedObject);
+		if(!FragmentItem->LocatorEditor.IsValid())
+		{
+			return FReply::Unhandled();
+		}
+
+		// Dropping onto an existing fragment
+		UObject* ResolvedObject = FragmentItem->LocatorEditor->ResolveDragOperation(DragOperation, nullptr);
+		if (!ResolvedObject)
+		{
+			return FReply::Unhandled();
+		}
+
+		FScopedTransaction Transaction(LOCTEXT("EditLocatorTransaction", "Edit Locator"));
+
+		PropertyHandle->NotifyPreChange();
+
+		TArray<void*> RawData;
+		PropertyHandle->AccessRawData(RawData);
+		for (void* Ptr : RawData)
+		{
+			FUniversalObjectLocator* Reference = static_cast<FUniversalObjectLocator*>(Ptr);
+			if(Reference->Fragments.IsValidIndex(FragmentItem->FragmentIndex))
+			{
+				// Remove all fragments beyond this one 
+				Reference->Fragments.RemoveAt(FragmentItem->FragmentIndex, Reference->Fragments.Num() - FragmentItem->FragmentIndex);
+
+				// Resolve to get the context for the new fragment
+				UObject* Context = nullptr;
+				if(Reference->Fragments.Num() > 0)
+				{
+					FResolveResult Result = Reference->Resolve(GetContext());
+					FResolveResultData ResultData = Result.SyncGet();
+					Context = ResultData.Object; 
+				}
+				else
+				{
+					Context = GetContext();
+				}
+				Reference->AddFragment(ResolvedObject, Context, nullptr);
+			}
+		}
+
+		TrimAbsoluteFragments();
+
+		PropertyHandle->NotifyPostChange(EPropertyChangeType::ValueSet);
+		PropertyHandle->NotifyFinishedChangingProperties();
+
+		return FReply::Handled();
 	}
+	else
+	{
+		// Dropping onto the tail
+		UObject* ResolvedObject = nullptr;
+
+		for (TPair<FName, TSharedPtr<ILocatorFragmentEditor>> Pair : ApplicableLocators)
+		{
+			ResolvedObject = Pair.Value->ResolveDragOperation(DragOperation, nullptr);
+			if (ResolvedObject)
+			{
+				break;
+			}
+		}
+
+		if (!ResolvedObject)
+		{
+			return FReply::Unhandled();
+		}
+
+		FScopedTransaction Transaction(LOCTEXT("EditLocatorTransaction", "Edit Locator"));
+
+		PropertyHandle->NotifyPreChange();
+
+		TArray<void*> RawData;
+		PropertyHandle->AccessRawData(RawData);
+		for (void* Ptr : RawData)
+		{
+			FUniversalObjectLocator* Reference = static_cast<FUniversalObjectLocator*>(Ptr);
 
-	return FReply::Handled();
+			// Resolve to get the context for the new fragment
+			UObject* Context = nullptr;
+			if(Reference->Fragments.Num() > 0)
+			{
+				FResolveResult Result = Reference->Resolve(GetContext());
+				FResolveResultData ResultData = Result.SyncGet();
+				Context = ResultData.Object; 
+			}
+			else
+			{
+				Context = GetContext();
+			}
+			Reference->AddFragment(ResolvedObject, Context, nullptr);
+		}
+
+		TrimAbsoluteFragments();
+
+		PropertyHandle->NotifyPostChange(EPropertyChangeType::ValueSet);
+		PropertyHandle->NotifyFinishedChangingProperties();
+
+		return FReply::Handled();
+	}
 }
 
 void FUniversalObjectLocatorCustomization::SetActor(AActor* NewObject)
@@ -363,96 +920,44 @@ void FUniversalObjectLocatorCustomization::SetActor(AActor* NewObject)
 	}
 }
 
-FUniversalObjectLocator* FUniversalObjectLocatorCustomization::GetSinglePropertyValue()
+FUniversalObjectLocator* FUniversalObjectLocatorCustomization::GetCommonPropertyValue()
 {
-	const FUniversalObjectLocator* ConstValue = static_cast<const FUniversalObjectLocatorCustomization*>(this)->GetSinglePropertyValue();
+	const FUniversalObjectLocator* ConstValue = static_cast<const FUniversalObjectLocatorCustomization*>(this)->GetCommonPropertyValue();
 	return const_cast<FUniversalObjectLocator*>(ConstValue);
 }
 
-const FUniversalObjectLocator* FUniversalObjectLocatorCustomization::GetSinglePropertyValue() const
+const FUniversalObjectLocator* FUniversalObjectLocatorCustomization::GetCommonPropertyValue() const
 {
-	const FUniversalObjectLocator* Value = nullptr;
-
+	TOptional<const FUniversalObjectLocator*> CommonValue;
 	PropertyHandle->EnumerateConstRawData(
-		[&Value](const void* Data, int32 DataIndex, int32 Num)
+		[&CommonValue](const void* Data, int32 DataIndex, int32 Num)
 		{
-			if (Num == 1)
+			if(!CommonValue.IsSet() && Data != nullptr)
 			{
-				Value = static_cast<const FUniversalObjectLocator*>(Data);
+				CommonValue = static_cast<const FUniversalObjectLocator*>(Data);
+			}
+			else if(Data == nullptr || *CommonValue.GetValue() != *static_cast<const FUniversalObjectLocator*>(Data))
+			{
+				CommonValue = nullptr;
+				return false;
 			}
-			return false;
+			return true;
 		}
 	);
 
-	return Value;
+	return CommonValue.Get(nullptr);
 }
 
 const FUniversalObjectLocatorCustomization::FCachedData& FUniversalObjectLocatorCustomization::GetCachedData() const
 {
-
-	const FUniversalObjectLocator* SingleValue = nullptr;
-
 	bool bNeedsUpdate = false;
-	if (CachedData.LocatorEditorType == NAME_None || !CachedData.PropertyValue.IsSet())
+	if (!CachedData.PropertyValue.IsSet())
 	{
-		// Set the LocatorEditorType type based on the current values
-		TOptional<FName> CommonLocatorEditorType;
-
-		PropertyHandle->EnumerateConstRawData(
-			[this, &SingleValue, &CommonLocatorEditorType](const void* Data, int32 DataIndex, int32 Num)
-			{
-				const FUniversalObjectLocator* Value = static_cast<const FUniversalObjectLocator*>(Data);
-				if (Num == 1)
-				{
-					SingleValue = Value;
-				}
-
-				const FFragmentType* FragmentTypePtr = Value->GetLastFragmentType();
-				if (!CommonLocatorEditorType)
-				{
-					if (FragmentTypePtr)
-					{
-						CommonLocatorEditorType = FragmentTypePtr->PrimaryEditorType;
-					}
-					else
-					{
-						// If we find any value that is effectively null, we choose the most applicable locator
-						return false;
-					}
-				}
-				else if (!FragmentTypePtr || FragmentTypePtr->PrimaryEditorType != CommonLocatorEditorType.GetValue())
-				{
-					CommonLocatorEditorType = NAME_None;
-				}
-				else
-				{
-					CommonLocatorEditorType = FragmentTypePtr->PrimaryEditorType;
-				}
-				return true;
-			}
-		);
-
-		if (CommonLocatorEditorType)
+		if (const FUniversalObjectLocator* CommonValue = GetCommonPropertyValue())
 		{
-			if (CommonLocatorEditorType.GetValue() != NAME_None)
+			if (!CachedData.PropertyValue.IsSet() || *CommonValue != CachedData.PropertyValue.GetValue())
 			{
-				CachedData.LocatorEditorType = CommonLocatorEditorType.GetValue();
-			}
-		}
-		else if (CachedData.LocatorEditorType == NAME_None)
-		{
-			for (const TPair<FName, TSharedPtr<ILocatorEditor>>& Pair : ApplicableLocators)
-			{
-				CachedData.LocatorEditorType = Pair.Key;
-				break;
-			}
-		}
-
-		if (SingleValue)
-		{
-			if (!CachedData.PropertyValue.IsSet() || *SingleValue != CachedData.PropertyValue.GetValue())
-			{
-				CachedData.PropertyValue = *SingleValue;
+				CachedData.PropertyValue = *CommonValue;
 				bNeedsUpdate = true;
 			}
 		}
@@ -461,7 +966,6 @@ const FUniversalObjectLocatorCustomization::FCachedData& FUniversalObjectLocator
 			CachedData.PropertyValue.Reset();
 			bNeedsUpdate = true;
 		}
-
 	}
 
 	if (bNeedsUpdate)
@@ -483,49 +987,6 @@ const FUniversalObjectLocatorCustomization::FCachedData& FUniversalObjectLocator
 		{
 			CachedData.ObjectPath.Empty();
 		}
-
-		TArray<void*> RawData;
-		PropertyHandle->AccessRawData(RawData);
-
-		bool bAnyFragmentType = false;
-		const FFragmentType* CommonFragmentType = nullptr;
-		TOptional<FFragmentTypeHandle> CommonFragmentTypeHandle;
-		for (void* Ptr : RawData)
-		{
-			FUniversalObjectLocator* Locator      = static_cast<FUniversalObjectLocator*>(Ptr);
-			const FFragmentType*     FragmentType = Locator->GetLastFragmentType();
-
-			if (CommonFragmentType && FragmentType != CommonFragmentType)
-			{
-				CommonFragmentType = nullptr;
-				CommonFragmentTypeHandle.Reset();
-				break;
-			}
-
-			if (FragmentType)
-			{
-				bAnyFragmentType = true;
-			}
-			CommonFragmentType = FragmentType;
-			CommonFragmentTypeHandle = Locator->GetLastFragmentTypeHandle();
-		}
-
-		FText NoneText = LOCTEXT("NoValues", "None");
-		if (CommonFragmentType != nullptr)
-		{
-			CachedData.FragmentTypeText = CommonFragmentType->DisplayText;
-			CachedData.FragmentType     = CommonFragmentTypeHandle;
-		}
-		else if (!bAnyFragmentType)
-		{
-			CachedData.FragmentTypeText = NoneText;
-		}
-		else
-		{
-			CachedData.FragmentTypeText = LOCTEXT("MultipleTypes", "<multiple>");
-		}
-
-		Rebuild();
 	}
 
 	return CachedData;
diff --git a/Engine/Source/Editor/UniversalObjectLocatorEditor/Private/UniversalObjectLocatorCustomization.h b/Engine/Source/Editor/UniversalObjectLocatorEditor/Private/UniversalObjectLocatorCustomization.h
--- a/Engine/Source/Editor/UniversalObjectLocatorEditor/Private/UniversalObjectLocatorCustomization.h
+++ b/Engine/Source/Editor/UniversalObjectLocatorEditor/Private/UniversalObjectLocatorCustomization.h
@@ -2,6 +2,7 @@
 
 #pragma once
 
+#include "EditorUndoClient.h"
 #include "IPropertyTypeCustomization.h"
 #include "Internationalization/Text.h"
 #include "Templates/SharedPointer.h"
@@ -10,6 +11,7 @@
 
 #include "UniversalObjectLocator.h"
 #include "IUniversalObjectLocatorCustomization.h"
+#include "Layout/Visibility.h"
 
 struct FGeometry;
 struct FAssetData;
@@ -26,23 +28,25 @@ class AActor;
 
 class SBox;
 class SWidget;
+class SWrapBox;
 
 namespace UE::UniversalObjectLocator
 {
 
-class ILocatorEditor;
+class ILocatorFragmentEditor;
+struct FFragmentItem;
 
-struct FUniversalObjectLocatorCustomization final : public IPropertyTypeCustomization, public IUniversalObjectLocatorCustomization
+struct FUniversalObjectLocatorCustomization final : public IPropertyTypeCustomization, public IUniversalObjectLocatorCustomization, public FSelfRegisteringEditorUndoClient
 {
 	static TSharedRef<IPropertyTypeCustomization> MakeInstance();
 
 	void CustomizeHeader(TSharedRef<IPropertyHandle> StructPropertyHandle, FDetailWidgetRow& HeaderRow, IPropertyTypeCustomizationUtils& StructCustomizationUtils) override;
 	void CustomizeChildren(TSharedRef<IPropertyHandle> StructPropertyHandle, IDetailChildrenBuilder& StructBuilder, IPropertyTypeCustomizationUtils& StructCustomizationUtils) override;
 
-	UObject* GetContext() const override { return nullptr; }
+	UObject* GetContext() const override { return WeakContext.Get(); }
 	UObject* GetSingleObject() const override;
 	FString GetPathToObject() const override;
-	void SetValue(FUniversalObjectLocator&& InNewValue) const override;
+	void SetValue(FUniversalObjectLocator&& InNewValue) override;
 	TSharedPtr<IPropertyHandle> GetProperty() const override;
 
 private:
@@ -52,40 +56,57 @@ private:
 		TOptional<FUniversalObjectLocator> PropertyValue;
 		TWeakObjectPtr<> WeakObject;
 		FString ObjectPath;
-		TOptional<FFragmentTypeHandle> FragmentType;
-		FText FragmentTypeText;
-		FName LocatorEditorType;
 	};
 
 private:
 
-	void Rebuild() const;
+	virtual void PostUndo(bool bSuccess) override { Rebuild(); }
+	virtual void PostRedo(bool bSuccess) override { Rebuild(); }
 
-	TSharedRef<SWidget> GetUserExposedFragmentTypeList();
-	FText GetCurrentFragmentTypeText() const;
+	void Rebuild();
 
-	bool HandleIsDragAllowed(TSharedPtr<FDragDropOperation> InDragOperation);
-	FReply HandleDrop(const FGeometry& InGeometry, const FDragDropEvent& InDropEvent);
+	void TrimAbsoluteFragments();
+
+	TSharedRef<SWidget> GetUserExposedFragmentTypeList(TWeakPtr<FFragmentItem> InWeakFragmentItem);
+	FText GetFragmentText(TWeakPtr<FFragmentItem> InWeakFragmentItem) const;
+	FText GetFragmentTooltipText(TWeakPtr<FFragmentItem> InWeakFragmentItem) const;
+	const FSlateBrush* GetFragmentIcon(TWeakPtr<FFragmentItem> InWeakFragmentItem) const;
+
+	TSharedRef<SWidget> GetFragmentTypeWidget(TWeakPtr<FFragmentItem> InWeakFragmentItem);
+
+	bool HandleIsDragAllowed(TSharedPtr<FDragDropOperation> InDragOperation, TWeakPtr<FFragmentItem> InWeakFragmentItem);
+	FReply HandleDrop(const FGeometry& InGeometry, const FDragDropEvent& InDropEvent, TWeakPtr<FFragmentItem> InWeakFragmentItem);
 
 	void SetActor(AActor* InActor);
 
 	const FCachedData& GetCachedData() const;
 
-	FUniversalObjectLocator* GetSinglePropertyValue();
-	const FUniversalObjectLocator* GetSinglePropertyValue() const;
+	FUniversalObjectLocator* GetCommonPropertyValue();
+	const FUniversalObjectLocator* GetCommonPropertyValue() const;
+
+	void ChangeEditorType(TWeakPtr<ILocatorFragmentEditor> InNewLocatorEditor, TWeakPtr<FFragmentItem> InWeakFragmentItem);
+	bool CompareCurrentEditorType(TWeakPtr<ILocatorFragmentEditor> InNewLocatorEditor, TWeakPtr<FFragmentItem> InWeakFragmentItem) const;
 
-	void ChangeEditorType(FName InNewEditor);
-	bool CompareCurrentEditorType(FName InNewEditor) const;
+	void RemoveFragment(TWeakPtr<FFragmentItem> InWeakFragmentItem);
+	void ClearFragments();
 
 private:
 
 	TSharedPtr<IPropertyUtilities> PropertyUtilities;
 	TSharedPtr<IPropertyHandle> PropertyHandle;
-	TSharedPtr<SBox> Content;
 
 	mutable FCachedData CachedData;
 
-	TMap<FName, TSharedPtr<ILocatorEditor>> ApplicableLocators;
+	TMap<FName, TSharedPtr<ILocatorFragmentEditor>> ApplicableLocators;
+
+	mutable TArray<TSharedRef<FFragmentItem>> Fragments;
+
+	TSharedPtr<SWrapBox> WrapBox;
+
+	TWeakObjectPtr<> WeakContext;
+	TWeakObjectPtr<UClass> WeakContextClass;
+
+	friend struct FFragmentItem;
 };
 
 
diff --git a/Engine/Source/Editor/UniversalObjectLocatorEditor/Private/UniversalObjectLocatorEditor.cpp b/Engine/Source/Editor/UniversalObjectLocatorEditor/Private/UniversalObjectLocatorEditor.cpp
--- a/Engine/Source/Editor/UniversalObjectLocatorEditor/Private/UniversalObjectLocatorEditor.cpp
+++ b/Engine/Source/Editor/UniversalObjectLocatorEditor/Private/UniversalObjectLocatorEditor.cpp
@@ -1,11 +1,14 @@
 // Copyright Epic Games, Inc. All Rights Reserved.
-#include "UniversalObjectLocatorEditor.h"
+
+#include "UniversalObjectLocatorFragmentEditor.h"
 #include "UniversalObjectLocator.h"
 
 namespace UE::UniversalObjectLocator
 {
-	FUniversalObjectLocator ILocatorEditor::MakeDefaultLocator() const
+	UClass* ILocatorFragmentEditor::ResolveClass(const FUniversalObjectLocatorFragment& InFragment, UObject* InContext) const
 	{
-		return FUniversalObjectLocator();
+		const FResolveParams ResolveParams(InContext);
+		UObject* Object = InFragment.Resolve(ResolveParams).SyncGet().Object;
+		return Object != nullptr ? Object->GetClass() : nullptr;
 	}
 }
\ No newline at end of file
diff --git a/Engine/Source/Editor/UniversalObjectLocatorEditor/Private/UniversalObjectLocatorEditorModule.cpp b/Engine/Source/Editor/UniversalObjectLocatorEditor/Private/UniversalObjectLocatorEditorModule.cpp
--- a/Engine/Source/Editor/UniversalObjectLocatorEditor/Private/UniversalObjectLocatorEditorModule.cpp
+++ b/Engine/Source/Editor/UniversalObjectLocatorEditor/Private/UniversalObjectLocatorEditorModule.cpp
@@ -31,7 +31,7 @@ void FUniversalObjectLocatorEditorModule::ShutdownModule()
 {
 }
 
-void FUniversalObjectLocatorEditorModule::RegisterLocatorEditor(FName Name, TSharedPtr<ILocatorEditor> LocatorEditor)
+void FUniversalObjectLocatorEditorModule::RegisterLocatorEditor(FName Name, TSharedPtr<ILocatorFragmentEditor> LocatorEditor)
 {
 	LocatorEditors.Add(Name, LocatorEditor);
 }
@@ -41,6 +41,21 @@ void FUniversalObjectLocatorEditorModule::UnregisterLocatorEditor(FName Name)
 	LocatorEditors.Remove(Name);
 }
 
+TSharedPtr<ILocatorFragmentEditor> FUniversalObjectLocatorEditorModule::FindLocatorEditor(FName Name)
+{
+	return LocatorEditors.FindRef(Name);
+}
+
+void FUniversalObjectLocatorEditorModule::RegisterEditorContext(FName Name, TSharedPtr<ILocatorFragmentEditorContext> LocatorEditorContext)
+{
+	LocatorEditorContexts.Add(Name, LocatorEditorContext);
+}
+
+void FUniversalObjectLocatorEditorModule::UnregisterEditorContext(FName Name)
+{
+	LocatorEditorContexts.Remove(Name);
+}
+
 } // namespace UE::UniversalObjectLocator
 
 IMPLEMENT_MODULE(UE::UniversalObjectLocator::FUniversalObjectLocatorEditorModule, UniversalObjectLocatorEditor);
diff --git a/Engine/Source/Editor/UniversalObjectLocatorEditor/Private/UniversalObjectLocatorEditorModule.h b/Engine/Source/Editor/UniversalObjectLocatorEditor/Private/UniversalObjectLocatorEditorModule.h
--- a/Engine/Source/Editor/UniversalObjectLocatorEditor/Private/UniversalObjectLocatorEditorModule.h
+++ b/Engine/Source/Editor/UniversalObjectLocatorEditor/Private/UniversalObjectLocatorEditorModule.h
@@ -11,7 +11,7 @@
 namespace UE::UniversalObjectLocator
 {
 
-class ILocatorEditor;
+class ILocatorFragmentEditor;
 
 class FUniversalObjectLocatorEditorModule
 	: public IUniversalObjectLocatorEditorModule
@@ -22,11 +22,19 @@ public:
 
 	void ShutdownModule() override;
 
-	void RegisterLocatorEditor(FName LocatorName, TSharedPtr<ILocatorEditor> LocatorEditor) override;
+	void RegisterLocatorEditor(FName LocatorName, TSharedPtr<ILocatorFragmentEditor> LocatorEditor) override;
 
 	void UnregisterLocatorEditor(FName LocatorName) override;
 
-	TMap<FName, TSharedPtr<ILocatorEditor>> LocatorEditors;
+	TSharedPtr<ILocatorFragmentEditor> FindLocatorEditor(FName Name) override;
+
+	void RegisterEditorContext(FName Name, TSharedPtr<ILocatorFragmentEditorContext> LocatorEditorContext) override;
+
+	void UnregisterEditorContext(FName Name) override;
+
+	TMap<FName, TSharedPtr<ILocatorFragmentEditor>> LocatorEditors;
+
+	TMap<FName, TSharedPtr<ILocatorFragmentEditorContext>> LocatorEditorContexts;
 };
 
 } // namespace UE::UniversalObjectLocator
diff --git a/Engine/Source/Editor/UniversalObjectLocatorEditor/Public/IUniversalObjectLocatorCustomization.h b/Engine/Source/Editor/UniversalObjectLocatorEditor/Public/IUniversalObjectLocatorCustomization.h
--- a/Engine/Source/Editor/UniversalObjectLocatorEditor/Public/IUniversalObjectLocatorCustomization.h
+++ b/Engine/Source/Editor/UniversalObjectLocatorEditor/Public/IUniversalObjectLocatorCustomization.h
@@ -13,6 +13,24 @@ class IPropertyHandle;
 namespace UE::UniversalObjectLocator
 {
 
+class IFragmentEditorHandle
+{
+public:
+	virtual ~IFragmentEditorHandle() = default;
+
+	// Get the fragment that this handle refers to
+	virtual const FUniversalObjectLocatorFragment& GetFragment() const = 0;
+
+	// Get the context class that this fragment is resolved against
+	virtual const UClass* GetContextClass() const = 0;
+
+	// Get the class that this fragment is resolved to
+	virtual const UClass* GetResolvedClass() const = 0;
+
+	// Set the value of this fragment
+	virtual void SetValue(const FUniversalObjectLocatorFragment& InNewValue) = 0;
+};
+
 class IUniversalObjectLocatorCustomization
 {
 public:
@@ -21,7 +39,7 @@ public:
 	virtual UObject* GetContext() const = 0;
 	virtual UObject* GetSingleObject() const = 0;
 	virtual FString GetPathToObject() const = 0;
-	virtual void SetValue(FUniversalObjectLocator&& InNewValue) const = 0;
+	virtual void SetValue(FUniversalObjectLocator&& InNewValue) = 0;
 
 	virtual TSharedPtr<IPropertyHandle> GetProperty() const = 0;
 };
diff --git a/Engine/Source/Editor/UniversalObjectLocatorEditor/Public/IUniversalObjectLocatorEditorModule.h b/Engine/Source/Editor/UniversalObjectLocatorEditor/Public/IUniversalObjectLocatorEditorModule.h
--- a/Engine/Source/Editor/UniversalObjectLocatorEditor/Public/IUniversalObjectLocatorEditorModule.h
+++ b/Engine/Source/Editor/UniversalObjectLocatorEditor/Public/IUniversalObjectLocatorEditorModule.h
@@ -10,15 +10,34 @@ class FName;
 namespace UE::UniversalObjectLocator
 {
 
-class ILocatorEditor;
+class ILocatorFragmentEditor;
+class ILocatorFragmentEditorContext;
 
 class IUniversalObjectLocatorEditorModule
 	: public IModuleInterface
 {
 public:
+	// Register a locator editor
+	// @param Name                     The name of the locator editor
+	// @param LocatorEditor            The locator editor to register
+	virtual void RegisterLocatorEditor(FName Name, TSharedPtr<ILocatorFragmentEditor> LocatorEditor) = 0;
 
-	virtual void RegisterLocatorEditor(FName Name, TSharedPtr<ILocatorEditor> LocatorEditor) = 0;
+	// Unregister a locator editor previously registered via RegisterLocatorEditor
+	// @param Name                     The name of the locator editor
 	virtual void UnregisterLocatorEditor(FName Name) = 0;
+
+	// Find a locator editor previously registered via RegisterLocatorEditor
+	// @param Name                     The name of the locator editor
+	virtual TSharedPtr<ILocatorFragmentEditor> FindLocatorEditor(FName Name) = 0;
+
+	// Register a locator editor context, used to control editor behavior on a per-use basis
+	// @param Name                     The name of the context, used as an identifier
+	// @param LocatorEditorContext     The context to register
+	virtual void RegisterEditorContext(FName Name, TSharedPtr<ILocatorFragmentEditorContext> LocatorEditorContext) = 0;
+
+	// Unregister a locator editor context previously registered via RegisterEditorContext
+	// @param Name             The name of the context, used as an identifier
+	virtual void UnregisterEditorContext(FName Name) = 0;
 };
 
 
diff --git a/Engine/Source/Editor/UniversalObjectLocatorEditor/Public/UniversalObjectLocatorEditor.h b/Engine/Source/Editor/UniversalObjectLocatorEditor/Public/UniversalObjectLocatorEditor.h
deleted file mode 100644
diff --git a/Engine/Source/Editor/UniversalObjectLocatorEditor/Public/UniversalObjectLocatorEditorContext.h b/Engine/Source/Editor/UniversalObjectLocatorEditor/Public/UniversalObjectLocatorEditorContext.h
new file mode 100644
--- /dev/null
+++ b/Engine/Source/Editor/UniversalObjectLocatorEditor/Public/UniversalObjectLocatorEditorContext.h
@@ -0,0 +1,25 @@
+﻿// Copyright Epic Games, Inc. All Rights Reserved.
+
+#pragma once
+
+#include "Templates/SharedPointer.h"
+
+class UObject;
+class IPropertyHandle;
+
+namespace UE::UniversalObjectLocator
+{
+
+class UNIVERSALOBJECTLOCATOREDITOR_API ILocatorFragmentEditorContext : public TSharedFromThis<ILocatorFragmentEditorContext>
+{
+public:
+	virtual ~ILocatorFragmentEditorContext() = default;
+
+	// Get a context object that can be used to resolve partial fragments at edit time
+	virtual UObject* GetContext(const IPropertyHandle& InPropertyHandle) const = 0;
+
+	// Check whether the supplied locator fragment is supported in this context
+	virtual bool IsFragmentAllowed(FName InFragmentName) const { return true; }
+};
+
+} // namespace UE::UniversalObjectLocator
\ No newline at end of file
diff --git a/Engine/Source/Editor/UniversalObjectLocatorEditor/Public/UniversalObjectLocatorFragmentEditor.h b/Engine/Source/Editor/UniversalObjectLocatorEditor/Public/UniversalObjectLocatorFragmentEditor.h
new file mode 100644
--- /dev/null
+++ b/Engine/Source/Editor/UniversalObjectLocatorEditor/Public/UniversalObjectLocatorFragmentEditor.h
@@ -0,0 +1,101 @@
+// Copyright Epic Games, Inc. All Rights Reserved.
+
+#pragma once
+
+#include "Templates/SharedPointer.h"
+
+struct FSlateIcon;
+
+class FText;
+class UObject;
+class SWidget;
+class FDragDropOperation;
+class IPropertyHandle;
+class IDetailChildrenBuilder;
+class IPropertyTypeCustomizationUtils;
+struct FUniversalObjectLocator;
+struct FUniversalObjectLocatorFragment;
+
+namespace UE::UniversalObjectLocator
+{
+
+class IFragmentEditorHandle;
+class IUniversalObjectLocatorCustomization;
+
+enum class ELocatorFragmentEditorType : uint8
+{
+	// Locator fragment is relative to something else
+	Relative,
+
+	// Locator fragment is absolute
+	Absolute,
+};
+
+class UNIVERSALOBJECTLOCATOREDITOR_API ILocatorFragmentEditor : public TSharedFromThis<ILocatorFragmentEditor>
+{
+public:
+	virtual ~ILocatorFragmentEditor() = default;
+
+	// Get the type of this locator editor (relative, absolute)
+	// @return the locator editor type
+	virtual ELocatorFragmentEditorType GetLocatorFragmentEditorType() const = 0;
+
+	// Get whether this locator editor is allowed in the supplied context. @see ILocatorFragmentEditorContext.
+	// @param InContextName    The name of the context, if this is NAME_None, no context is supplied (the default for blueprint-instantiated UOL properties)
+	// @return true if this locator is allowed in the supplied context
+	virtual bool IsAllowedInContext(FName InContextName) const { return true; }
+	
+	// Called to check whether a drag operation is supported for this fragment
+	// @param InDragOperation  The drag operation to check
+	// @param InContext        The context object
+	// @return true if the operation is supported, false if unsupported
+	virtual bool IsDragSupported(TSharedPtr<FDragDropOperation> InDragOperation, UObject* InContext) const = 0;
+
+	// Called to resolve a drag operation to an object
+	// @param InDragOperation  The drag operation to resolve
+	// @param InContext        The context object
+	// @return the resolved object
+	virtual UObject* ResolveDragOperation(TSharedPtr<FDragDropOperation> InDragOperation, UObject* InContext) const = 0;
+
+	// Parameters used to create the UI for this locator fragment editor
+	struct FEditUIParameters
+	{
+		// The customization that is creating this UI
+		TSharedRef<IUniversalObjectLocatorCustomization> Customization;
+
+		// The handle to the fragment to create the UI for
+		TSharedRef<IFragmentEditorHandle> Handle;
+	};
+
+	// Make the editor UI for this fragment (displayed in a context menu)
+	// @param InParameters     Parameters used to create the UI
+	// @return the editing widget
+	virtual TSharedPtr<SWidget> MakeEditUI(const FEditUIParameters& InParameters) = 0;
+
+	// Get the text to display for the supplied fragment
+	// @param InFragment       The fragment to display. If this is nullptr then the name of the fragment type should be returned
+	// @return the text to display
+	virtual FText GetDisplayText(const FUniversalObjectLocatorFragment* InFragment = nullptr) const = 0;
+
+	// Get the tooltip text to display for this fragment
+	// @param InFragment       The fragment to display. If this is nullptr then the tooltip of the fragment type should be returned
+	// @return the tooltip text to display
+	virtual FText GetDisplayTooltip(const FUniversalObjectLocatorFragment* InFragment = nullptr) const = 0;
+
+	// Get the icon to display for this fragment
+	// @param InFragment       The fragment to display. If this is nullptr then the icon of the fragment type should be returned
+	// @return the icon to display
+	virtual FSlateIcon GetDisplayIcon(const FUniversalObjectLocatorFragment* InFragment = nullptr) const = 0;
+
+	// Resolve the output class of a fragment at edit time
+	// @param InFragment       The fragment to resolve the class of
+	// @param InContext        The context to resolve the class against
+	// @return the class that was resolved from the fragment
+	virtual UClass* ResolveClass(const FUniversalObjectLocatorFragment& InFragment, UObject* InContext) const;
+
+	// Make a default fragment for the given fragment type
+	// @return a default fragment for this locator editor to use
+	virtual FUniversalObjectLocatorFragment MakeDefaultLocatorFragment() const = 0;
+};
+
+} // namespace UE::UniversalObjectLocator
\ No newline at end of file
diff --git a/Engine/Source/Editor/UniversalObjectLocatorEditor/UniversalObjectLocatorEditor.Build.cs b/Engine/Source/Editor/UniversalObjectLocatorEditor/UniversalObjectLocatorEditor.Build.cs
--- a/Engine/Source/Editor/UniversalObjectLocatorEditor/UniversalObjectLocatorEditor.Build.cs
+++ b/Engine/Source/Editor/UniversalObjectLocatorEditor/UniversalObjectLocatorEditor.Build.cs
@@ -26,8 +26,10 @@ public class UniversalObjectLocatorEditor : ModuleRules
 				"SceneOutliner",
 				"Slate",
 				"SlateCore",
+				"InputCore",
 				"Sequencer",
 				"UniversalObjectLocator",
+				"UnrealEd",
 			}
 		);
 	}
diff --git a/Engine/Source/Runtime/Engine/Private/UnrealEngine.cpp b/Engine/Source/Runtime/Engine/Private/UnrealEngine.cpp
--- a/Engine/Source/Runtime/Engine/Private/UnrealEngine.cpp
+++ b/Engine/Source/Runtime/Engine/Private/UnrealEngine.cpp
@@ -338,7 +338,7 @@ void FEngineModule::StartupModule()
 			}
 			{
 				FFragmentTypeParameters FragmentTypeParams("animinst", NSLOCTEXT("Engine", "AnimInstanceLocatorFragment", "AnimInstance"));
-				FragmentTypeParams.PrimaryEditorType = "Component";
+				FragmentTypeParams.PrimaryEditorType = "AnimInstance";
 				FAnimInstanceLocatorFragment::FragmentType = UolModule.RegisterFragmentType<FAnimInstanceLocatorFragment>(FragmentTypeParams);
 			}
 			{
diff --git a/Engine/Source/Runtime/UniversalObjectLocator/Private/UniversalObjectLocator.cpp b/Engine/Source/Runtime/UniversalObjectLocator/Private/UniversalObjectLocator.cpp
--- a/Engine/Source/Runtime/UniversalObjectLocator/Private/UniversalObjectLocator.cpp
+++ b/Engine/Source/Runtime/UniversalObjectLocator/Private/UniversalObjectLocator.cpp
@@ -308,6 +308,18 @@ const FUniversalObjectLocatorFragment* FUniversalObjectLocator::GetLastFragment(
 	return Fragments.Num() != 0 ? &Fragments.Last() : nullptr;
 }
 
+bool FUniversalObjectLocator::ForEachFragment(TFunctionRef<bool(int32, int32, const FUniversalObjectLocatorFragment&)> InFunction) const
+{
+	for (int32 FragmentIndex = 0, NumFragments = Fragments.Num(); FragmentIndex < NumFragments; ++FragmentIndex)
+	{
+		if(!InFunction(FragmentIndex, NumFragments, Fragments[FragmentIndex]))
+		{
+			return false;
+		}
+	}
+
+	return true;
+}
 
 UE::UniversalObjectLocator::EFragmentTypeFlags FUniversalObjectLocator::GetDefaultFlags() const
 {
@@ -608,7 +620,7 @@ bool FUniversalObjectLocator::AddFragment(const UObject* Object, UObject* Contex
 	{
 		return false;
 	}
-
+	
 	// If the initialization needs to be relative to a different context, add a fragment for NewContext as well
 	if (Result.Type == ELocatorType::Relative && Result.RelativeToContext != Context)
 	{
diff --git a/Engine/Source/Runtime/UniversalObjectLocator/Public/UniversalObjectLocator.h b/Engine/Source/Runtime/UniversalObjectLocator/Public/UniversalObjectLocator.h
--- a/Engine/Source/Runtime/UniversalObjectLocator/Public/UniversalObjectLocator.h
+++ b/Engine/Source/Runtime/UniversalObjectLocator/Public/UniversalObjectLocator.h
@@ -9,6 +9,12 @@
 #include "UniversalObjectLocatorFragmentType.h"
 #include "UniversalObjectLocator.generated.h"
 
+namespace UE::UniversalObjectLocator
+{
+	struct FUniversalObjectLocatorCustomization;
+	struct FFragmentItem;
+}
+
 /**
  * Universal Object Locators (UOLs) define an address to an object.
  *
@@ -195,6 +201,12 @@ public:
 	 */
 	UNIVERSALOBJECTLOCATOR_API const FUniversalObjectLocatorFragment* GetLastFragment() const;
 
+	/**
+	 * Iterate the fragments in this address
+	 * @param InFunction       The visitor function to call for each fragment in this address. Function should return false to stop iteration.
+	 * @return false if the iteration was stopped early, true if the iteration succeeded (or if empty)
+	 */
+	UNIVERSALOBJECTLOCATOR_API bool ForEachFragment(TFunctionRef<bool(int32, int32, const FUniversalObjectLocatorFragment&)> InFunction) const;
 
 	/*
 	* Iterates over all fragments and combines their types' default flags.
@@ -225,9 +237,8 @@ public:
 	/*~ End TStructOpsTypeTraits implementation */
 
 private:
-
-	/**  */
-	bool AddFragment(const UObject* Object, UObject* Context, UObject* StopAtContext);
+	/** Adds the 'best' fragment according to the input object and context, as well as any potential required relative fragments */
+	UNIVERSALOBJECTLOCATOR_API bool AddFragment(const UObject* Object, UObject* Context, UObject* StopAtContext);
 
 	FResolveResult ResolveSyncImpl(const FResolveParams& Params) const;
 	FResolveResult ResolveAsyncImpl(const FResolveParams& Params) const;
@@ -237,6 +248,9 @@ private:
 	/** Array of relative locators ordered sequentially from outer to inner. The first locator is probably 'absolute' and is resolved with no context, although that is not a hard restriction */
 	UPROPERTY()
 	TArray<FUniversalObjectLocatorFragment> Fragments;
+
+	friend struct UE::UniversalObjectLocator::FUniversalObjectLocatorCustomization;
+	friend struct UE::UniversalObjectLocator::FFragmentItem;
 };
 
 template<typename FragmentType, typename ...ArgTypes>
