Defer rebuild of UOL customization

Fixes crashes where UOL fragment container widgets could be destroyed while still being referenced elsewhere

diff --git a/Engine/Source/Editor/UniversalObjectLocatorEditor/Private/UniversalObjectLocatorCustomization.cpp b/Engine/Source/Editor/UniversalObjectLocatorEditor/Private/UniversalObjectLocatorCustomization.cpp
--- a/Engine/Source/Editor/UniversalObjectLocatorEditor/Private/UniversalObjectLocatorCustomization.cpp
+++ b/Engine/Source/Editor/UniversalObjectLocatorEditor/Private/UniversalObjectLocatorCustomization.cpp
@@ -198,7 +198,7 @@ void FUniversalObjectLocatorCustomization::CustomizeHeader(TSharedRef<IPropertyH
 	.ValueContent()
 	.MaxDesiredWidth(400.0f)
 	[
-		SNew(SBorder)
+		SAssignNew(RootWidget, SBorder)
 		.ToolTipText_Lambda([this]()
 		{
 			FTextBuilder TextBuilder;
@@ -236,16 +236,31 @@ void FUniversalObjectLocatorCustomization::CustomizeHeader(TSharedRef<IPropertyH
 
 	PropertyHandle->SetOnPropertyValueChanged(FSimpleDelegate::CreateLambda([this]()
 	{
-		Rebuild();
+		RequestRebuild();
 	}));
 
-	Rebuild();
+	RequestRebuild();
 }
 
 void FUniversalObjectLocatorCustomization::CustomizeChildren(TSharedRef<IPropertyHandle> StructPropertyHandle, IDetailChildrenBuilder& StructBuilder, IPropertyTypeCustomizationUtils& StructCustomizationUtils)
 {
 }
 
+void FUniversalObjectLocatorCustomization::RequestRebuild()
+{
+	if(!bRebuildRequested)
+	{
+		bRebuildRequested = true;
+
+		RootWidget->RegisterActiveTimer(0.0f, FWidgetActiveTimerDelegate::CreateSPLambda(this, [this](double, float)
+		{
+			Rebuild();
+			bRebuildRequested = false;
+			return EActiveTimerReturnType::Stop;
+		}));
+	}
+}
+
 void FUniversalObjectLocatorCustomization::Rebuild()
 {
 	Fragments.Reset();
@@ -261,8 +276,8 @@ void FUniversalObjectLocatorCustomization::Rebuild()
 			TSharedRef<FFragmentItem> NewItem = MakeShared<FFragmentItem>();
 			NewItem->FragmentIndex = FragmentIndex;
 			NewItem->FragmentType = Fragment.GetFragmentTypeHandle();
-			NewItem->LocatorEditorType = Fragment.GetFragmentType()->PrimaryEditorType;
-			NewItem->LocatorEditor = ApplicableLocators.FindRef(Fragment.GetFragmentType()->PrimaryEditorType);
+			NewItem->LocatorEditorType = Fragment.GetFragmentType() ? Fragment.GetFragmentType()->PrimaryEditorType : NAME_None;
+			NewItem->LocatorEditor = ApplicableLocators.FindRef(NewItem->LocatorEditorType);
 			NewItem->PropertyHandle = PropertyHandle;
 			NewItem->WeakCustomization = SharedThis(this);
 			NewItem->bIsTail = FragmentIndex == CommonValue->Fragments.Num() - 1;
diff --git a/Engine/Source/Editor/UniversalObjectLocatorEditor/Private/UniversalObjectLocatorCustomization.h b/Engine/Source/Editor/UniversalObjectLocatorEditor/Private/UniversalObjectLocatorCustomization.h
--- a/Engine/Source/Editor/UniversalObjectLocatorEditor/Private/UniversalObjectLocatorCustomization.h
+++ b/Engine/Source/Editor/UniversalObjectLocatorEditor/Private/UniversalObjectLocatorCustomization.h
@@ -63,6 +63,7 @@ private:
 	virtual void PostUndo(bool bSuccess) override { Rebuild(); }
 	virtual void PostRedo(bool bSuccess) override { Rebuild(); }
 
+	void RequestRebuild();
 	void Rebuild();
 
 	void TrimAbsoluteFragments();
@@ -101,11 +102,14 @@ private:
 
 	mutable TArray<TSharedRef<FFragmentItem>> Fragments;
 
+	TSharedPtr<SWidget> RootWidget;
 	TSharedPtr<SWrapBox> WrapBox;
 
 	TWeakObjectPtr<> WeakContext;
 	TWeakObjectPtr<UClass> WeakContextClass;
 
+	bool bRebuildRequested = false;
+
 	friend struct FFragmentItem;
 };
 
