Sequencer: Fix tooltip crash on changing binding type.

diff --git a/Engine/Source/Editor/UniversalObjectLocatorEditor/Private/UniversalObjectLocatorCustomization.cpp b/Engine/Source/Editor/UniversalObjectLocatorEditor/Private/UniversalObjectLocatorCustomization.cpp
--- a/Engine/Source/Editor/UniversalObjectLocatorEditor/Private/UniversalObjectLocatorCustomization.cpp
+++ b/Engine/Source/Editor/UniversalObjectLocatorEditor/Private/UniversalObjectLocatorCustomization.cpp
@@ -32,6 +32,8 @@
 #include "String/ParseTokens.h"
 #include "Widgets/Layout/SWrapBox.h"
 #include "Widgets/Views/STileView.h"
+#include "IDetailChildrenBuilder.h"
+#include "DetailLayoutBuilder.h"
 
 
 #define LOCTEXT_NAMESPACE "FUniversalObjectLocatorCustomization"
@@ -244,6 +246,13 @@ void FUniversalObjectLocatorCustomization::CustomizeHeader(TSharedRef<IPropertyH
 
 void FUniversalObjectLocatorCustomization::CustomizeChildren(TSharedRef<IPropertyHandle> StructPropertyHandle, IDetailChildrenBuilder& StructBuilder, IPropertyTypeCustomizationUtils& StructCustomizationUtils)
 {
+	if (IDetailsView* DetailsView = StructBuilder.GetParentCategory().GetParentLayout().GetDetailsView())
+	{
+		DetailsView->OnFinishedChangingProperties().AddLambda([this](const FPropertyChangedEvent&)
+			{
+				RequestRebuild();
+			});
+	}
 }
 
 void FUniversalObjectLocatorCustomization::RequestRebuild()
@@ -709,9 +718,10 @@ FText FUniversalObjectLocatorCustomization::GetFragmentText(TWeakPtr<FFragmentIt
 
 FText FUniversalObjectLocatorCustomization::GetFragmentTooltipText(TWeakPtr<FFragmentItem> InWeakFragmentItem) const
 {
+	// If a rebuild has been requested, return a blank item, as we need to rebuild before we grab the tooltip text,  otherwise we might be dealing with stale or garbage data
 	const FText NoneText = LOCTEXT("NoValues", "None");
 	TSharedPtr<FFragmentItem> FragmentItem = InWeakFragmentItem.Pin();
-	if(!FragmentItem.IsValid())
+	if(!FragmentItem.IsValid() || bRebuildRequested)
 	{
 		return NoneText;
 	}
