Sequencer: Fix further crash with changing binding type.

diff --git a/Engine/Source/Editor/UniversalObjectLocatorEditor/Private/UniversalObjectLocatorCustomization.cpp b/Engine/Source/Editor/UniversalObjectLocatorEditor/Private/UniversalObjectLocatorCustomization.cpp
--- a/Engine/Source/Editor/UniversalObjectLocatorEditor/Private/UniversalObjectLocatorCustomization.cpp
+++ b/Engine/Source/Editor/UniversalObjectLocatorEditor/Private/UniversalObjectLocatorCustomization.cpp
@@ -248,7 +248,7 @@ void FUniversalObjectLocatorCustomization::CustomizeChildren(TSharedRef<IPropert
 {
 	if (IDetailsView* DetailsView = StructBuilder.GetParentCategory().GetParentLayout().GetDetailsView())
 	{
-		DetailsView->OnFinishedChangingProperties().AddLambda([this](const FPropertyChangedEvent&)
+		DetailsView->OnFinishedChangingProperties().AddSPLambda(this, [this](const FPropertyChangedEvent&)
 			{
 				RequestRebuild();
 			});
