Hide the 'use selected' item in the actor picker for binding UOLs

diff --git a/Engine/Source/Editor/UniversalObjectLocatorEditor/Private/ActorLocatorEditor.cpp b/Engine/Source/Editor/UniversalObjectLocatorEditor/Private/ActorLocatorEditor.cpp
--- a/Engine/Source/Editor/UniversalObjectLocatorEditor/Private/ActorLocatorEditor.cpp
+++ b/Engine/Source/Editor/UniversalObjectLocatorEditor/Private/ActorLocatorEditor.cpp
@@ -89,6 +89,7 @@ TSharedPtr<SWidget> FActorLocatorEditor::MakeEditUI(const FEditUIParameters& InP
 	AActor* InitialActor = GetActor(InParameters.Handle);
 	const bool bAllowClear = true;
 	const bool bAllowPickingLevelInstanceContent = true;
+	const bool bDontDisplayUseSelected = false;
 	const FOnActorSelected OnActorSelected = FOnActorSelected::CreateSP(this, &FActorLocatorEditor::OnSetActor, TWeakPtr<IFragmentEditorHandle>(InParameters.Handle));
 	const FOnShouldFilterActor OnShouldFilterActor = FOnShouldFilterActor::CreateLambda([](const AActor*){ return true; });
 
@@ -104,7 +105,8 @@ TSharedPtr<SWidget> FActorLocatorEditor::MakeEditUI(const FEditUIParameters& InP
 				OnShouldFilterActor,
 				OnActorSelected,
 				FSimpleDelegate(),
-				FSimpleDelegate())
+				FSimpleDelegate(),
+				bDontDisplayUseSelected)
 		];
 }
 
