Fix UOL not always wrapping correctly in AnimNext

diff --git a/Engine/Source/Editor/UniversalObjectLocatorEditor/Private/UniversalObjectLocatorCustomization.cpp b/Engine/Source/Editor/UniversalObjectLocatorEditor/Private/UniversalObjectLocatorCustomization.cpp
--- a/Engine/Source/Editor/UniversalObjectLocatorEditor/Private/UniversalObjectLocatorCustomization.cpp
+++ b/Engine/Source/Editor/UniversalObjectLocatorEditor/Private/UniversalObjectLocatorCustomization.cpp
@@ -232,7 +232,7 @@ void FUniversalObjectLocatorCustomization::CustomizeHeader(TSharedRef<IPropertyH
 		.Padding(2.0f, 2.0f, 2.0f, 2.0f)
 		[
 			SAssignNew(WrapBox, SWrapBox)
-			.PreferredSize(400.0f)
+			.UseAllottedSize(true)
 		]
 	];
 
